Merging Darkside's StandaloneRenderer changes.

- Due to not using vendor headers when cross-compiling, we cannot use deprecated calls like glTexEnvi(). This needs to be later reworked to use another (not-plaform-specific) #define.

diff --git a/Engine/Shaders/StandaloneRenderer/OpenGL/SlateElementPixelShader.glsl b/Engine/Shaders/StandaloneRenderer/OpenGL/SlateElementPixelShader.glsl
--- a/Engine/Shaders/StandaloneRenderer/OpenGL/SlateElementPixelShader.glsl
+++ b/Engine/Shaders/StandaloneRenderer/OpenGL/SlateElementPixelShader.glsl
@@ -30,7 +30,11 @@ varying vec4 Color;
 vec4 GetFontElementColor()
 {
 	vec4 OutColor = Color;
+#if PLATFORM_LINUX
+	OutColor.a *= texture2D(ElementTexture, TexCoords.xy).r; // OpenGL 3.2+ uses Red for single channel textures
+#else
 	OutColor.a *= texture2D(ElementTexture, TexCoords.xy).a;
+#endif
 
 	return OutColor;
 }
diff --git a/Engine/Source/Developer/StandaloneRenderer/Private/Linux/OpenGL/SlateOpenGLContext.cpp b/Engine/Source/Developer/StandaloneRenderer/Private/Linux/OpenGL/SlateOpenGLContext.cpp
--- a/Engine/Source/Developer/StandaloneRenderer/Private/Linux/OpenGL/SlateOpenGLContext.cpp
+++ b/Engine/Source/Developer/StandaloneRenderer/Private/Linux/OpenGL/SlateOpenGLContext.cpp
@@ -3,14 +3,23 @@
 #include "StandaloneRendererPrivate.h"
 #include "OpenGL/SlateOpenGLRenderer.h"
 
+static SDL_Window* CreateDummyGLWindow()
+{
+	SDL_Init( SDL_INIT_VIDEO );
 
-static SDL_Window* CreateDummyGLWindow();
+	// Create a dummy window.
+	SDL_Window *h_wnd = SDL_CreateWindow(	NULL,
+		0, 0, 1, 1,
+		SDL_WINDOW_OPENGL | SDL_WINDOW_BORDERLESS | SDL_WINDOW_HIDDEN );
 
+	return h_wnd;
+}
 
 FSlateOpenGLContext::FSlateOpenGLContext()
+	: WindowHandle(NULL)
+	, Context(NULL)
+	, bReleaseWindowOnDestroy(false)
 {
-	WindowHandle	= NULL;
-	OpenGLContext	= NULL;
 }
 
 FSlateOpenGLContext::~FSlateOpenGLContext()
@@ -28,44 +37,39 @@ void FSlateOpenGLContext::Initialize( void* InWindow, const FSlateOpenGLContext*
 		bReleaseWindowOnDestroy = true;
 	}
 
-#if 0
-	//	Create a pixel format descriptor for this window
-	SDL_PixelFormat PFD;
-	FMemory::Memzero( &PFD, sizeof(SDL_PixelFormat) );
-
-	PFD.format	= SDL_PIXELFORMAT_ARGB8888;
-	PFD.palette = NULL;
-	PFD.BitsPerPixel = 32;
-	PFD.BytesPerPixel = 4;
-	PFD.Amask = 0xFF000000;
-	PFD.Rmask = 0x00FF0000;
-	PFD.Gmask = 0x0000FF00;
-	PFD.Bmask = 0x000099FF;
-
-	SDL_SetWindowPixelFormat;
-	SDL_SetPixelFormat
-#endif
-
-	OpenGLContext = SDL_GL_CreateContext( WindowHandle );
-	SDL_GL_MakeCurrent( WindowHandle, OpenGLContext );
+	SDL_GL_SetAttribute( SDL_GL_CONTEXT_MAJOR_VERSION, 3 );
+	SDL_GL_SetAttribute( SDL_GL_CONTEXT_MINOR_VERSION, 2 );
+	SDL_GL_SetAttribute( SDL_GL_CONTEXT_FLAGS, SDL_GL_CONTEXT_FORWARD_COMPATIBLE_FLAG);
+	SDL_GL_SetAttribute( SDL_GL_CONTEXT_PROFILE_MASK, SDL_GL_CONTEXT_PROFILE_COMPATIBILITY );
+
+	if( SharedContext )
+	{
+		SDL_GL_SetAttribute( SDL_GL_SHARE_WITH_CURRENT_CONTEXT, 1 );
+		SDL_GL_MakeCurrent( SharedContext->WindowHandle, SharedContext->Context );
+	}
+	else
+	{
+		SDL_GL_SetAttribute( SDL_GL_SHARE_WITH_CURRENT_CONTEXT, 0 );
+	}
+
+	Context = SDL_GL_CreateContext( WindowHandle );
+	SDL_GL_MakeCurrent( WindowHandle, Context );
 }
 
 void FSlateOpenGLContext::Destroy()
 {
-	if	( WindowHandle )
+	if	( WindowHandle != NULL )
 	{
-		SDL_GL_DeleteContext( OpenGLContext );
 		SDL_GL_MakeCurrent( NULL, NULL );
+		SDL_GL_DeleteContext( Context );
 
 		if	( bReleaseWindowOnDestroy )
 		{
 			SDL_DestroyWindow( WindowHandle );
+			SDL_Quit();
 		}
-
-		SDL_Quit();
-
 		WindowHandle = NULL;
-		OpenGLContext = NULL;
+
 	}
 }
 
@@ -73,23 +77,12 @@ void FSlateOpenGLContext::MakeCurrent()
 {
 	if	( WindowHandle )
 	{
-		SDL_GL_MakeCurrent( WindowHandle, OpenGLContext );
+		CHECK_GL_ERRORS;
+		if(SDL_GL_MakeCurrent( WindowHandle, Context ) == 0)
+		{
+			glGetError(); // SDL leaves glGetError in a dirty state even when successful?
+		}
 	}
 }
 
-static SDL_Window* CreateDummyGLWindow()
-{
-	SDL_Window	*h_wnd;
-
-	SDL_Init( SDL_INIT_VIDEO );
-
-	h_wnd = SDL_CreateWindow(	"SlateViewer",				// window title
-								SDL_WINDOWPOS_UNDEFINED,	// initial x position
-								SDL_WINDOWPOS_UNDEFINED,	// initial y position
-								1280,						// width, in pixels
-								 800,						// height, in pixels
-								SDL_WINDOW_OPENGL | SDL_WINDOW_RESIZABLE );
-
-	return h_wnd;
-}
 
diff --git a/Engine/Source/Developer/StandaloneRenderer/Private/Linux/OpenGL/SlateOpenGLViewport.cpp b/Engine/Source/Developer/StandaloneRenderer/Private/Linux/OpenGL/SlateOpenGLViewport.cpp
--- a/Engine/Source/Developer/StandaloneRenderer/Private/Linux/OpenGL/SlateOpenGLViewport.cpp
+++ b/Engine/Source/Developer/StandaloneRenderer/Private/Linux/OpenGL/SlateOpenGLViewport.cpp
@@ -4,8 +4,8 @@
 
 #include "OpenGL/SlateOpenGLRenderer.h"
 
-
 FSlateOpenGLViewport::FSlateOpenGLViewport()
+:	bFullscreen(false)
 {
 }
 
diff --git a/Engine/Source/Developer/StandaloneRenderer/Private/OpenGL/SlateOpenGLExtensions.cpp b/Engine/Source/Developer/StandaloneRenderer/Private/OpenGL/SlateOpenGLExtensions.cpp
--- a/Engine/Source/Developer/StandaloneRenderer/Private/OpenGL/SlateOpenGLExtensions.cpp
+++ b/Engine/Source/Developer/StandaloneRenderer/Private/OpenGL/SlateOpenGLExtensions.cpp
@@ -58,6 +58,11 @@ PFNGLGETATTRIBLOCATIONPROC			glGetAttribLocation			= NULL;
 
 PFNWGLCREATECONTEXTATTRIBSARBPROC	wglCreateContextAttribsARB	= NULL;
 
+#elif PLATFORM_LINUX
+
+#define DEFINE_GL_ENTRYPOINTS(Type,Func) Type Func = NULL;
+ENUM_GL_ENTRYPOINTS(DEFINE_GL_ENTRYPOINTS);
+
 #endif
 
 bool bOpenGLExtensionsLoaded;
@@ -115,6 +120,9 @@ void LoadOpenGLExtensions()
 		glGetAttachedShaders		=(PFNGLGETATTACHEDSHADERSPROC)		wglGetProcAddress("glGetAttachedShaders");
 		glGetAttribLocation			=(PFNGLGETATTRIBLOCATIONPROC)		wglGetProcAddress("glGetAttribLocation");
 	}
+#elif PLATFORM_LINUX
+#define GET_GL_ENTRYPOINTS(Type,Func) Func = (Type)SDL_GL_GetProcAddress(#Func);
+	ENUM_GL_ENTRYPOINTS(GET_GL_ENTRYPOINTS);
 #endif
 	// If extensions are needed for your platform add support for them here 
 }
diff --git a/Engine/Source/Developer/StandaloneRenderer/Private/OpenGL/SlateOpenGLExtensions.h b/Engine/Source/Developer/StandaloneRenderer/Private/OpenGL/SlateOpenGLExtensions.h
--- a/Engine/Source/Developer/StandaloneRenderer/Private/OpenGL/SlateOpenGLExtensions.h
+++ b/Engine/Source/Developer/StandaloneRenderer/Private/OpenGL/SlateOpenGLExtensions.h
@@ -57,6 +57,57 @@ extern PFNGLGETATTRIBLOCATIONPROC			glGetAttribLocation;
 
 extern PFNWGLCREATECONTEXTATTRIBSARBPROC	wglCreateContextAttribsARB;
 
+#elif PLATFORM_LINUX
+
+#define ENUM_GL_ENTRYPOINTS(EnumMacro) \
+	EnumMacro(PFNGLBINDTEXTUREPROC,glBindTexture) \
+	EnumMacro(PFNGLBLENDFUNCPROC,glBlendFunc) \
+	EnumMacro(PFNGLDELETETEXTURESPROC,glDeleteTextures) \
+	EnumMacro(PFNGLDISABLEPROC,glDisable) \
+	EnumMacro(PFNGLENABLEPROC,glEnable) \
+	EnumMacro(PFNGLGENTEXTURESPROC,glGenTextures) \
+	EnumMacro(PFNGLGETERRORPROC,glGetError) \
+	EnumMacro(PFNGLPOLYGONMODEPROC,glPolygonMode) \
+	EnumMacro(PFNGLSTENCILFUNCPROC,glStencilFunc) \
+	EnumMacro(PFNGLSTENCILMASKPROC,glStencilMask) \
+	EnumMacro(PFNGLSTENCILOPPROC,glStencilOp) \
+	EnumMacro(PFNGLTEXIMAGE2DPROC,glTexImage2D) \
+	EnumMacro(PFNGLTEXPARAMETERIPROC,glTexParameteri) \
+	EnumMacro(PFNGLVIEWPORTPROC,glViewport) \
+	EnumMacro(PFNGLBLENDEQUATIONPROC,glBlendEquation) \
+	EnumMacro(PFNGLDRAWRANGEELEMENTSPROC,glDrawRangeElements) \
+	EnumMacro(PFNGLACTIVETEXTUREPROC,glActiveTexture) \
+	EnumMacro(PFNGLBINDBUFFERPROC,glBindBuffer) \
+	EnumMacro(PFNGLDELETEBUFFERSPROC,glDeleteBuffers) \
+	EnumMacro(PFNGLGENBUFFERSPROC,glGenBuffers) \
+	EnumMacro(PFNGLBUFFERDATAPROC,glBufferData) \
+	EnumMacro(PFNGLMAPBUFFERPROC,glMapBuffer) \
+	EnumMacro(PFNGLUNMAPBUFFERPROC,glUnmapBuffer) \
+	EnumMacro(PFNGLATTACHSHADERPROC,glAttachShader) \
+	EnumMacro(PFNGLBINDATTRIBLOCATIONPROC,glBindAttribLocation) \
+	EnumMacro(PFNGLCOMPILESHADERPROC,glCompileShader) \
+	EnumMacro(PFNGLCREATEPROGRAMPROC,glCreateProgram) \
+	EnumMacro(PFNGLCREATESHADERPROC,glCreateShader) \
+	EnumMacro(PFNGLDELETEPROGRAMPROC,glDeleteProgram) \
+	EnumMacro(PFNGLDELETESHADERPROC,glDeleteShader) \
+	EnumMacro(PFNGLENABLEVERTEXATTRIBARRAYPROC,glEnableVertexAttribArray) \
+	EnumMacro(PFNGLGETPROGRAMIVPROC,glGetProgramiv) \
+	EnumMacro(PFNGLGETPROGRAMINFOLOGPROC,glGetProgramInfoLog) \
+	EnumMacro(PFNGLGETSHADERIVPROC,glGetShaderiv) \
+	EnumMacro(PFNGLGETSHADERINFOLOGPROC,glGetShaderInfoLog) \
+	EnumMacro(PFNGLGETUNIFORMLOCATIONPROC,glGetUniformLocation) \
+	EnumMacro(PFNGLLINKPROGRAMPROC,glLinkProgram) \
+	EnumMacro(PFNGLSHADERSOURCEPROC,glShaderSource) \
+	EnumMacro(PFNGLUSEPROGRAMPROC,glUseProgram) \
+	EnumMacro(PFNGLUNIFORM4FPROC,glUniform4f) \
+	EnumMacro(PFNGLUNIFORM1IPROC,glUniform1i) \
+	EnumMacro(PFNGLUNIFORM4FVPROC,glUniform4fv) \
+	EnumMacro(PFNGLUNIFORMMATRIX4FVPROC,glUniformMatrix4fv) \
+	EnumMacro(PFNGLVERTEXATTRIBPOINTERPROC,glVertexAttribPointer) \
+
+#define DECLARE_GL_ENTRYPOINTS(Type,Func) extern Type Func;
+ENUM_GL_ENTRYPOINTS(DECLARE_GL_ENTRYPOINTS);
+
 #endif
 
 void LoadOpenGLExtensions();
diff --git a/Engine/Source/Developer/StandaloneRenderer/Private/OpenGL/SlateOpenGLIndexBuffer.cpp b/Engine/Source/Developer/StandaloneRenderer/Private/OpenGL/SlateOpenGLIndexBuffer.cpp
--- a/Engine/Source/Developer/StandaloneRenderer/Private/OpenGL/SlateOpenGLIndexBuffer.cpp
+++ b/Engine/Source/Developer/StandaloneRenderer/Private/OpenGL/SlateOpenGLIndexBuffer.cpp
@@ -5,7 +5,8 @@
 #include "OpenGL/SlateOpenGLIndexBuffer.h"
 
 FSlateOpenGLIndexBuffer::FSlateOpenGLIndexBuffer()
-	: MaxNumIndices(0)	 
+	: MaxNumIndices(0)
+	, BufferID(0)
 {
 
 }
diff --git a/Engine/Source/Developer/StandaloneRenderer/Private/OpenGL/SlateOpenGLRenderer.h b/Engine/Source/Developer/StandaloneRenderer/Private/OpenGL/SlateOpenGLRenderer.h
--- a/Engine/Source/Developer/StandaloneRenderer/Private/OpenGL/SlateOpenGLRenderer.h
+++ b/Engine/Source/Developer/StandaloneRenderer/Private/OpenGL/SlateOpenGLRenderer.h
@@ -40,6 +40,10 @@ struct FSlateOpenGLContext
 #elif PLATFORM_IOS
 	UIWindow* WindowHandle;
 	EAGLContext* Context;
+#elif PLATFORM_LINUX
+	SDL_Window* WindowHandle;
+	SDL_GLContext Context;
+	bool bReleaseWindowOnDestroy;
 #endif
 
 	FSlateOpenGLContext();
diff --git a/Engine/Source/Developer/StandaloneRenderer/Private/OpenGL/SlateOpenGLRenderingPolicy.cpp b/Engine/Source/Developer/StandaloneRenderer/Private/OpenGL/SlateOpenGLRenderingPolicy.cpp
--- a/Engine/Source/Developer/StandaloneRenderer/Private/OpenGL/SlateOpenGLRenderingPolicy.cpp
+++ b/Engine/Source/Developer/StandaloneRenderer/Private/OpenGL/SlateOpenGLRenderingPolicy.cpp
@@ -140,12 +140,12 @@ void FSlateOpenGLRenderingPolicy::DrawElements( const FMatrix& ViewProjectionMat
 	// OpenGL state toggles
 	glDisable(GL_DEPTH_TEST);
 	glDisable(GL_CULL_FACE);
-#if !PLATFORM_USES_ES2
+#if !PLATFORM_USES_ES2 && !PLATFORM_LINUX
 	glEnable(GL_TEXTURE_2D);
 #endif
 
 	// Set up alpha testing
-#if !PLATFORM_USES_ES2
+#if !PLATFORM_USES_ES2 && !PLATFORM_LINUX
 	glEnable(GL_ALPHA_TEST);
 	glAlphaFunc( GL_GREATER, 0.0f );
 #endif
diff --git a/Engine/Source/Developer/StandaloneRenderer/Private/OpenGL/SlateOpenGLShaders.cpp b/Engine/Source/Developer/StandaloneRenderer/Private/OpenGL/SlateOpenGLShaders.cpp
--- a/Engine/Source/Developer/StandaloneRenderer/Private/OpenGL/SlateOpenGLShaders.cpp
+++ b/Engine/Source/Developer/StandaloneRenderer/Private/OpenGL/SlateOpenGLShaders.cpp
@@ -93,9 +93,11 @@ void FSlateOpenGLShader::CompileShader( const FString& Filename, GLenum ShaderTy
 	ANSICHAR* Chars[2] = {0};
 	// pass the #define along to the shader
 #if PLATFORM_USES_ES2
-	Chars[0] = (ANSICHAR*)"#define PLATFORM_USES_ES2 1\n";
+	Chars[0] = (ANSICHAR*)"#define PLATFORM_USES_ES2 1\n\n#define PLATFORM_LINUX 0\n";
+#elif PLATFORM_LINUX
+	Chars[0] = (ANSICHAR*)"#version 150\n\n#define PLATFORM_USES_ES2 0\n\n#define PLATFORM_LINUX 1\n";
 #else
-	Chars[0] = (ANSICHAR*)"#version 120\n\n#define PLATFORM_USES_ES2 0\n";
+	Chars[0] = (ANSICHAR*)"#version 120\n\n#define PLATFORM_USES_ES2 0\n\n#define PLATFORM_LINUX 0\n";
 #endif
 	Chars[1] = new ANSICHAR[Source.Len()+1];
 	FCStringAnsi::Strcpy(Chars[1], Source.Len() + 1, TCHAR_TO_ANSI(*Source));
diff --git a/Engine/Source/Developer/StandaloneRenderer/Private/OpenGL/SlateOpenGLTextures.cpp b/Engine/Source/Developer/StandaloneRenderer/Private/OpenGL/SlateOpenGLTextures.cpp
--- a/Engine/Source/Developer/StandaloneRenderer/Private/OpenGL/SlateOpenGLTextures.cpp
+++ b/Engine/Source/Developer/StandaloneRenderer/Private/OpenGL/SlateOpenGLTextures.cpp
@@ -13,13 +13,15 @@ void FSlateOpenGLTexture::Init( GLenum TexFormat, const TArray<uint8>& TextureDa
 	CHECK_GL_ERRORS;
 
 	// Ensure texturing is enabled before setting texture properties
-#if !PLATFORM_USES_ES2
+#if !PLATFORM_USES_ES2 && !PLATFORM_LINUX
 	glEnable(GL_TEXTURE_2D);
 #endif
 	glBindTexture(GL_TEXTURE_2D, ShaderResource);
 	glTexParameteri(GL_TEXTURE_2D, GL_TEXTURE_MIN_FILTER, GL_LINEAR);
 	glTexParameteri(GL_TEXTURE_2D, GL_TEXTURE_MAG_FILTER, GL_LINEAR);
+#if !PLATFORM_LINUX
 	glTexEnvi( GL_TEXTURE_ENV, GL_TEXTURE_ENV_MODE, GL_REPLACE );
+#endif
 
 	// the raw data is in bgra or bgr
 	const GLint Format = GL_BGRA;
@@ -86,7 +88,11 @@ void FSlateFontTextureOpenGL::ConditionalUpdateTexture()
 		glTexParameteri(GL_TEXTURE_2D, GL_TEXTURE_STORAGE_HINT_APPLE, GL_STORAGE_CACHED_APPLE);
 		glPixelStorei(GL_UNPACK_CLIENT_STORAGE_APPLE, GL_TRUE);
 #endif
+#if PLATFORM_LINUX
+		glTexImage2D( GL_TEXTURE_2D, 0, GL_R8, AtlasWidth, AtlasHeight, 0, GL_RED, GL_UNSIGNED_BYTE, AtlasData.GetData() );
+#else
 		glTexImage2D( GL_TEXTURE_2D, 0, Format, AtlasWidth, AtlasHeight, 0, Format, GL_UNSIGNED_BYTE, AtlasData.GetData() );
+#endif
 #if PLATFORM_MAC
 		glPixelStorei(GL_UNPACK_CLIENT_STORAGE_APPLE, GL_FALSE);
 #endif
diff --git a/Engine/Source/Developer/StandaloneRenderer/Private/OpenGL/SlateOpenGLVertexBuffer.cpp b/Engine/Source/Developer/StandaloneRenderer/Private/OpenGL/SlateOpenGLVertexBuffer.cpp
--- a/Engine/Source/Developer/StandaloneRenderer/Private/OpenGL/SlateOpenGLVertexBuffer.cpp
+++ b/Engine/Source/Developer/StandaloneRenderer/Private/OpenGL/SlateOpenGLVertexBuffer.cpp
@@ -7,6 +7,7 @@
 FSlateOpenGLVertexBuffer::FSlateOpenGLVertexBuffer( uint32 InStride )
 	: BufferSize(0)
 	, Stride(InStride)
+	, BufferID(0)
 {
 	check( InStride > 0 );
 }
diff --git a/Engine/Source/Developer/StandaloneRenderer/Private/StandaloneRendererPrivate.h b/Engine/Source/Developer/StandaloneRenderer/Private/StandaloneRendererPrivate.h
--- a/Engine/Source/Developer/StandaloneRenderer/Private/StandaloneRendererPrivate.h
+++ b/Engine/Source/Developer/StandaloneRenderer/Private/StandaloneRendererPrivate.h
@@ -67,7 +67,13 @@
 #define glTexEnvi(...)
 #define glAlphaFunc(...)
 
-#endif // PLATFORM_IOS
+#elif PLATFORM_LINUX
+//#define GLCOREARB_PROTOTYPES 
+#include <GL/glcorearb.h>
+//#include <GL/glext.h>
+#include "SDL.h"
+
+#endif // PLATFORM_LINUX
 
 #include "OpenGL/SlateOpenGLExtensions.h"
 #include "OpenGL/SlateOpenGLTextures.h"
