Fix standalone Slate tooltips on OS X from appearing as garbage since the work done to improve the Editor performance. The standalone renderer needed the same code to perform context updates & temporarily disable/enable screen updates. Also fixed a bug in the MacWindow code that could result in setting an incorrect frame.

diff --git a/Engine/Source/Developer/StandaloneRenderer/Private/Mac/OpenGL/SlateOpenGLContext.cpp b/Engine/Source/Developer/StandaloneRenderer/Private/Mac/OpenGL/SlateOpenGLContext.cpp
--- a/Engine/Source/Developer/StandaloneRenderer/Private/Mac/OpenGL/SlateOpenGLContext.cpp
+++ b/Engine/Source/Developer/StandaloneRenderer/Private/Mac/OpenGL/SlateOpenGLContext.cpp
@@ -90,12 +90,16 @@ char const* const CompositedBlitFragmentShader = "#version 120\n"
 			self = [super initWithFrame: FrameRect];
 		}
 	}
+	
+	[[NSNotificationCenter defaultCenter] addObserver:self selector:@selector(_surfaceNeedsUpdate:) name:NSViewGlobalFrameDidChangeNotification object:self];
 
 	return self;
 }
 
 - (void)dealloc
 {
+	[[NSNotificationCenter defaultCenter] removeObserver:self name:NSViewGlobalFrameDidChangeNotification object:self];
+	
 	if (Context)
 	{
 		[Context release];
@@ -203,6 +207,26 @@ char const* const CompositedBlitFragmentShader = "#version 120\n"
 	}
 }
 
+- (void)renewGState
+{
+	extern bool GMacEnableCocoaScreenUpdates;
+	if(GMacEnableCocoaScreenUpdates)
+	{
+		GMacEnableCocoaScreenUpdates = false;
+		NSDisableScreenUpdates();
+	}
+	
+	[super renewGState];
+}
+
+- (void) _surfaceNeedsUpdate:(NSNotification*)notification
+{
+	if(Context)
+	{
+		[Context update];
+	}
+}
+
 @end
 
 
diff --git a/Engine/Source/Developer/StandaloneRenderer/Private/Mac/OpenGL/SlateOpenGLViewport.cpp b/Engine/Source/Developer/StandaloneRenderer/Private/Mac/OpenGL/SlateOpenGLViewport.cpp
--- a/Engine/Source/Developer/StandaloneRenderer/Private/Mac/OpenGL/SlateOpenGLViewport.cpp
+++ b/Engine/Source/Developer/StandaloneRenderer/Private/Mac/OpenGL/SlateOpenGLViewport.cpp
@@ -6,6 +6,8 @@
 
 #include "MacWindow.h"
 
+/** Used to temporarily disable Cocoa screen updates to make window updates happen only on the render thread. */
+bool GMacEnableCocoaScreenUpdates = true;
 
 FSlateOpenGLViewport::FSlateOpenGLViewport()
 {
@@ -120,6 +122,12 @@ void FSlateOpenGLViewport::SwapBuffers()
 	
 	[RenderingContext.Context flushBuffer];
 	
+	if(!GMacEnableCocoaScreenUpdates)
+	{
+		GMacEnableCocoaScreenUpdates = true;
+		NSEnableScreenUpdates();
+	}
+	
 	NSWindow* Window = [[RenderingContext.Context view] window];
 	FSlateCocoaWindow* SlateCocoaWindow = [Window isKindOfClass:[FSlateCocoaWindow class]] ? (FSlateCocoaWindow*)Window : nil;
 	if(SlateCocoaWindow)
diff --git a/Engine/Source/Runtime/Core/Private/Mac/MacWindow.cpp b/Engine/Source/Runtime/Core/Private/Mac/MacWindow.cpp
--- a/Engine/Source/Runtime/Core/Private/Mac/MacWindow.cpp
+++ b/Engine/Source/Runtime/Core/Private/Mac/MacWindow.cpp
@@ -91,7 +91,7 @@
 	if(bRenderInitialised && (bDeferSetFrame || bDeferSetOrigin))
 	{
 		dispatch_block_t Block = ^{
-			if(bDeferSetOrigin)
+			if(!bDeferSetFrame && bDeferSetOrigin)
 			{
 				DeferFrame.size = [self frame].size;
 			}
