Linux: Fix standalone programs after Friday changes.

Temporary, hopefully will go away soon together with DSExt.

diff --git a/Engine/Source/Developer/StandaloneRenderer/Private/Linux/OpenGL/SlateOpenGLContext.cpp b/Engine/Source/Developer/StandaloneRenderer/Private/Linux/OpenGL/SlateOpenGLContext.cpp
--- a/Engine/Source/Developer/StandaloneRenderer/Private/Linux/OpenGL/SlateOpenGLContext.cpp
+++ b/Engine/Source/Developer/StandaloneRenderer/Private/Linux/OpenGL/SlateOpenGLContext.cpp
@@ -2,6 +2,7 @@
 
 #include "StandaloneRendererPrivate.h"
 #include "OpenGL/SlateOpenGLRenderer.h"
+#include "ds_extensions.h"
 
 static SDL_Window* CreateDummyGLWindow()
 {
@@ -12,6 +13,12 @@ static SDL_Window* CreateDummyGLWindow()
 		0, 0, 1, 1,
 		SDL_WINDOW_OPENGL | SDL_WINDOW_BORDERLESS | SDL_WINDOW_HIDDEN );
 
+    int CacheRet = DSEXT_CacheX11Info(h_wnd);
+    if (EDSExtSuccess != CacheRet)
+    {
+        UE_LOG(LogInit, Error, TEXT("Could not cache X11 info, DSExt error: %d"), CacheRet);
+    }
+
 	return h_wnd;
 }
 
