Update interfaces for standalone renderers

diff --git a/Engine/Source/Developer/StandaloneRenderer/Private/OpenGL/SlateOpenGLRenderingPolicy.cpp b/Engine/Source/Developer/StandaloneRenderer/Private/OpenGL/SlateOpenGLRenderingPolicy.cpp
--- a/Engine/Source/Developer/StandaloneRenderer/Private/OpenGL/SlateOpenGLRenderingPolicy.cpp
+++ b/Engine/Source/Developer/StandaloneRenderer/Private/OpenGL/SlateOpenGLRenderingPolicy.cpp
@@ -283,13 +283,3 @@ void FSlateOpenGLRenderingPolicy::DrawElements( const FMatrix& ViewProjectionMat
 	glUseProgram(0);
 
 }
-
-/**
- * Returns a slate texture resource
- *
- * @param	The name of the texture to return
- */
-FSlateShaderResourceProxy* FSlateOpenGLRenderingPolicy::GetTextureResource( const FSlateBrush& Brush )
-{
-	return TextureManager->GetTexture( Brush );
-}
diff --git a/Engine/Source/Developer/StandaloneRenderer/Private/OpenGL/SlateOpenGLRenderingPolicy.h b/Engine/Source/Developer/StandaloneRenderer/Private/OpenGL/SlateOpenGLRenderingPolicy.h
--- a/Engine/Source/Developer/StandaloneRenderer/Private/OpenGL/SlateOpenGLRenderingPolicy.h
+++ b/Engine/Source/Developer/StandaloneRenderer/Private/OpenGL/SlateOpenGLRenderingPolicy.h
@@ -27,23 +27,12 @@ public:
 	 */
 	void DrawElements( const FMatrix& ViewProjectionMatrix, const TArray<FSlateRenderBatch>& RenderBatches );
 
-	
-	/**
-	 * Returns a slate texture that represents a viewport                   
-	 */
-	FSlateShaderResource* GetViewportResource( const ISlateViewport* InViewportInterface ) { return NULL; }
-	
-	/**
-	 * Returns a slate texture resource
-	 *
-	 * @param	The name of the texture to return
- 	 */
-	FSlateShaderResourceProxy* GetTextureResource( const FSlateBrush& Brush );
 
+	virtual TSharedRef<FSlateShaderResourceManager> GetResourceManager() override { return TextureManager.ToSharedRef(); }
 	/**
 	 * Returns the font cache used when the OpenGL rendering policy is active
  	 */
-	TSharedPtr<FSlateFontCache>& GetFontCache() { return FontCache; }
+	virtual TSharedRef<FSlateFontCache> GetFontCache() override { return FontCache.ToSharedRef(); }
 
 	/** 
 	 * Initializes resources if needed
diff --git a/Engine/Source/Developer/StandaloneRenderer/Private/OpenGL/SlateOpenGLTextureManager.cpp b/Engine/Source/Developer/StandaloneRenderer/Private/OpenGL/SlateOpenGLTextureManager.cpp
--- a/Engine/Source/Developer/StandaloneRenderer/Private/OpenGL/SlateOpenGLTextureManager.cpp
+++ b/Engine/Source/Developer/StandaloneRenderer/Private/OpenGL/SlateOpenGLTextureManager.cpp
@@ -298,7 +298,7 @@ void FSlateOpenGLTextureManager::ReleaseDynamicTextureResource(const FSlateBrush
 /**
  * Returns a texture with the passed in name or NULL if it cannot be found.
  */
-FSlateShaderResourceProxy* FSlateOpenGLTextureManager::GetTexture( const FSlateBrush& InBrush )
+FSlateShaderResourceProxy* FSlateOpenGLTextureManager::GetShaderResource( const FSlateBrush& InBrush )
 {
 	FSlateShaderResourceProxy* Texture = NULL;
 	if( InBrush.IsDynamicallyLoaded() )
diff --git a/Engine/Source/Developer/StandaloneRenderer/Private/OpenGL/SlateOpenGLTextureManager.h b/Engine/Source/Developer/StandaloneRenderer/Private/OpenGL/SlateOpenGLTextureManager.h
--- a/Engine/Source/Developer/StandaloneRenderer/Private/OpenGL/SlateOpenGLTextureManager.h
+++ b/Engine/Source/Developer/StandaloneRenderer/Private/OpenGL/SlateOpenGLTextureManager.h
@@ -20,7 +20,7 @@ public:
 	/**
 	 * Returns a texture with the passed in name or NULL if it cannot be found.
 	 */
-	virtual FSlateShaderResourceProxy* GetTexture( const FSlateBrush& InBrush );
+	virtual FSlateShaderResourceProxy* GetShaderResource( const FSlateBrush& InBrush ) override;
 
 	/**
 	 * Creates a 1x1 texture of the specified color
diff --git a/Engine/Source/Developer/StandaloneRenderer/Private/Windows/D3D/SlateD3DRenderingPolicy.cpp b/Engine/Source/Developer/StandaloneRenderer/Private/Windows/D3D/SlateD3DRenderingPolicy.cpp
--- a/Engine/Source/Developer/StandaloneRenderer/Private/Windows/D3D/SlateD3DRenderingPolicy.cpp
+++ b/Engine/Source/Developer/StandaloneRenderer/Private/Windows/D3D/SlateD3DRenderingPolicy.cpp
@@ -290,9 +290,3 @@ void FSlateD3D11RenderingPolicy::DrawElements( const FMatrix& ViewProjectionMatr
 
 	}
 }
-
-
-FSlateShaderResourceProxy* FSlateD3D11RenderingPolicy::GetTextureResource( const FSlateBrush& Brush )
-{
-	return TextureManager->GetTexture( Brush );
-}
diff --git a/Engine/Source/Developer/StandaloneRenderer/Private/Windows/D3D/SlateD3DRenderingPolicy.h b/Engine/Source/Developer/StandaloneRenderer/Private/Windows/D3D/SlateD3DRenderingPolicy.h
--- a/Engine/Source/Developer/StandaloneRenderer/Private/Windows/D3D/SlateD3DRenderingPolicy.h
+++ b/Engine/Source/Developer/StandaloneRenderer/Private/Windows/D3D/SlateD3DRenderingPolicy.h
@@ -1,8 +1,7 @@
 // Copyright 1998-2014 Epic Games, Inc. All Rights Reserved.
 
 
-#ifndef __SLATE_D3D_RENDERING_POLICY_H__
-#define __SLATE_D3D_RENDERING_POLICY_H__
+#pragma once
 
 class FSlateD3D11RenderingPolicy : public FSlateRenderingPolicy
 {
@@ -13,9 +12,8 @@ public:
 	void UpdateBuffers( const FSlateWindowElementList& InElementList );
 	void DrawElements( const FMatrix& ViewProjectionMatrix, const TArray<FSlateRenderBatch>& RenderBatches );
 
-	FSlateShaderResource* GetViewportResource( const ISlateViewport* InViewportInterface ) { return NULL; }
-	FSlateShaderResourceProxy* GetTextureResource( const FSlateBrush& Brush );
-	TSharedPtr<FSlateFontCache>& GetFontCache() { return FontCache; }
+	virtual TSharedRef<FSlateShaderResourceManager> GetResourceManager() override { return TextureManager; }
+	virtual TSharedRef<FSlateFontCache> GetFontCache() override { return FontCache.ToSharedRef(); }
 
 private:
 	void InitResources();
@@ -39,5 +37,3 @@ private:
 	TSharedPtr<FSlateFontCache> FontCache;
 	TSharedRef<FSlateD3DTextureManager> TextureManager;
 };
-
-#endif // __SLATE_D3D_RENDERING_POLICY_H__
diff --git a/Engine/Source/Developer/StandaloneRenderer/Private/Windows/D3D/SlateD3DTextureManager.cpp b/Engine/Source/Developer/StandaloneRenderer/Private/Windows/D3D/SlateD3DTextureManager.cpp
--- a/Engine/Source/Developer/StandaloneRenderer/Private/Windows/D3D/SlateD3DTextureManager.cpp
+++ b/Engine/Source/Developer/StandaloneRenderer/Private/Windows/D3D/SlateD3DTextureManager.cpp
@@ -134,7 +134,7 @@ void FSlateD3DTextureManager::CreateTextures( const TArray< const FSlateBrush* >
 void FSlateD3DTextureManager::CreateTextureNoAtlas( const FSlateBrush& InBrush )
 {
 	const FName TextureName = InBrush.GetResourceName();
-	if( TextureName != NAME_None && GetTexture( InBrush ) == NULL )
+	if( TextureName != NAME_None && GetShaderResource( InBrush ) == NULL )
 	{
 		FNewTextureInfo Info;
 		Info.bShouldAtlas = false;
@@ -159,7 +159,7 @@ void FSlateD3DTextureManager::CreateTextureNoAtlas( const FSlateBrush& InBrush )
 /**
  * Returns a texture with the passed in name or NULL if it cannot be found.
  */
-FSlateShaderResourceProxy* FSlateD3DTextureManager::GetTexture( const FSlateBrush& InBrush )
+FSlateShaderResourceProxy* FSlateD3DTextureManager::GetShaderResource( const FSlateBrush& InBrush )
 {
 	FSlateShaderResourceProxy* Texture = NULL;
 	if( InBrush.IsDynamicallyLoaded() )
diff --git a/Engine/Source/Developer/StandaloneRenderer/Private/Windows/D3D/SlateD3DTextureManager.h b/Engine/Source/Developer/StandaloneRenderer/Private/Windows/D3D/SlateD3DTextureManager.h
--- a/Engine/Source/Developer/StandaloneRenderer/Private/Windows/D3D/SlateD3DTextureManager.h
+++ b/Engine/Source/Developer/StandaloneRenderer/Private/Windows/D3D/SlateD3DTextureManager.h
@@ -24,7 +24,7 @@ public:
 	/**
 	 * Returns a texture with the passed in name or NULL if it cannot be found.
 	 */
-	virtual FSlateShaderResourceProxy* GetTexture( const FSlateBrush& InBrush );
+	virtual FSlateShaderResourceProxy* GetShaderResource( const FSlateBrush& InBrush ) override;
 
 	void CreateTextureNoAtlas( const FSlateBrush& InBrush );
 
diff --git a/Engine/Source/Runtime/SlateCore/Public/Rendering/RenderingPolicy.h b/Engine/Source/Runtime/SlateCore/Public/Rendering/RenderingPolicy.h
--- a/Engine/Source/Runtime/SlateCore/Public/Rendering/RenderingPolicy.h
+++ b/Engine/Source/Runtime/SlateCore/Public/Rendering/RenderingPolicy.h
@@ -35,7 +35,7 @@ public:
 
 	virtual void UpdateBuffers( const FSlateWindowElementList& WindowElementList ) = 0;
 
-	virtual TSharedPtr<class FSlateFontCache> GetFontCache( ) = 0;
+	virtual TSharedRef<class FSlateFontCache> GetFontCache( ) = 0;
 
 	virtual TSharedRef<class FSlateShaderResourceManager> GetResourceManager() = 0;
 
diff --git a/Engine/Source/Runtime/SlateRHIRenderer/Private/SlateRHIRenderingPolicy.h b/Engine/Source/Runtime/SlateRHIRenderer/Private/SlateRHIRenderingPolicy.h
--- a/Engine/Source/Runtime/SlateRHIRenderer/Private/SlateRHIRenderingPolicy.h
+++ b/Engine/Source/Runtime/SlateRHIRenderer/Private/SlateRHIRenderingPolicy.h
@@ -179,8 +179,8 @@ public:
 	virtual void UpdateBuffers( const FSlateWindowElementList& WindowElementList ) override;
 	virtual void DrawElements( const FIntPoint& InViewportSize, class FSlateBackBuffer& BackBuffer, const FMatrix& ViewProjectionMatrix, const TArray<FSlateRenderBatch>& RenderBatches );
 
-	TSharedPtr<FSlateFontCache> GetFontCache() { return FontCache; }
-	TSharedRef<FSlateShaderResourceManager> GetResourceManager() { return ResourceManager; }
+	virtual TSharedRef<FSlateFontCache> GetFontCache() override { return FontCache.ToSharedRef(); }
+	virtual TSharedRef<FSlateShaderResourceManager> GetResourceManager() override { return ResourceManager; }
 
 	void InitResources();
 	void ReleaseResources();
