UnrealEngineLauncher non-unity Win & Mac build fixes

diff --git a/Engine/Source/Developer/StandaloneRenderer/Private/OpenGL/SlateOpenGLRenderingPolicy.cpp b/Engine/Source/Developer/StandaloneRenderer/Private/OpenGL/SlateOpenGLRenderingPolicy.cpp
--- a/Engine/Source/Developer/StandaloneRenderer/Private/OpenGL/SlateOpenGLRenderingPolicy.cpp
+++ b/Engine/Source/Developer/StandaloneRenderer/Private/OpenGL/SlateOpenGLRenderingPolicy.cpp
@@ -283,3 +283,8 @@ void FSlateOpenGLRenderingPolicy::DrawElements( const FMatrix& ViewProjectionMat
 	glUseProgram(0);
 
 }
+
+TSharedRef<FSlateShaderResourceManager> FSlateOpenGLRenderingPolicy::GetResourceManager()
+{
+	return TextureManager.ToSharedRef();
+}
diff --git a/Engine/Source/Developer/StandaloneRenderer/Private/OpenGL/SlateOpenGLRenderingPolicy.h b/Engine/Source/Developer/StandaloneRenderer/Private/OpenGL/SlateOpenGLRenderingPolicy.h
--- a/Engine/Source/Developer/StandaloneRenderer/Private/OpenGL/SlateOpenGLRenderingPolicy.h
+++ b/Engine/Source/Developer/StandaloneRenderer/Private/OpenGL/SlateOpenGLRenderingPolicy.h
@@ -28,7 +28,7 @@ public:
 	void DrawElements( const FMatrix& ViewProjectionMatrix, const TArray<FSlateRenderBatch>& RenderBatches );
 
 
-	virtual TSharedRef<FSlateShaderResourceManager> GetResourceManager() override { return TextureManager.ToSharedRef(); }
+	virtual TSharedRef<FSlateShaderResourceManager> GetResourceManager() override;
 	/**
 	 * Returns the font cache used when the OpenGL rendering policy is active
  	 */
diff --git a/Engine/Source/Developer/StandaloneRenderer/Private/Windows/D3D/SlateD3DRenderingPolicy.cpp b/Engine/Source/Developer/StandaloneRenderer/Private/Windows/D3D/SlateD3DRenderingPolicy.cpp
--- a/Engine/Source/Developer/StandaloneRenderer/Private/Windows/D3D/SlateD3DRenderingPolicy.cpp
+++ b/Engine/Source/Developer/StandaloneRenderer/Private/Windows/D3D/SlateD3DRenderingPolicy.cpp
@@ -290,3 +290,9 @@ void FSlateD3D11RenderingPolicy::DrawElements( const FMatrix& ViewProjectionMatr
 
 	}
 }
+
+TSharedRef<FSlateShaderResourceManager> FSlateD3D11RenderingPolicy::GetResourceManager()
+{
+	return TextureManager;
+}
+
diff --git a/Engine/Source/Developer/StandaloneRenderer/Private/Windows/D3D/SlateD3DRenderingPolicy.h b/Engine/Source/Developer/StandaloneRenderer/Private/Windows/D3D/SlateD3DRenderingPolicy.h
--- a/Engine/Source/Developer/StandaloneRenderer/Private/Windows/D3D/SlateD3DRenderingPolicy.h
+++ b/Engine/Source/Developer/StandaloneRenderer/Private/Windows/D3D/SlateD3DRenderingPolicy.h
@@ -12,7 +12,7 @@ public:
 	void UpdateBuffers( const FSlateWindowElementList& InElementList );
 	void DrawElements( const FMatrix& ViewProjectionMatrix, const TArray<FSlateRenderBatch>& RenderBatches );
 
-	virtual TSharedRef<FSlateShaderResourceManager> GetResourceManager() override { return TextureManager; }
+	virtual TSharedRef<FSlateShaderResourceManager> GetResourceManager() override;
 	virtual TSharedRef<FSlateFontCache> GetFontCache() override { return FontCache.ToSharedRef(); }
 
 private:
