Fix for opengl/d3d slate standalone renderer compile errors

diff --git a/Engine/Source/Developer/StandaloneRenderer/Private/OpenGL/SlateOpenGLRenderer.h b/Engine/Source/Developer/StandaloneRenderer/Private/OpenGL/SlateOpenGLRenderer.h
--- a/Engine/Source/Developer/StandaloneRenderer/Private/OpenGL/SlateOpenGLRenderer.h
+++ b/Engine/Source/Developer/StandaloneRenderer/Private/OpenGL/SlateOpenGLRenderer.h
@@ -140,6 +140,7 @@ public:
 	virtual void CreateViewport( const TSharedRef<SWindow> InWindow ) override;
 	virtual void RequestResize( const TSharedPtr<SWindow>& InWindow, uint32 NewSizeX, uint32 NewSizeY ) override;
 	virtual void UpdateFullscreenState( const TSharedRef<SWindow> InWindow, uint32 OverrideResX, uint32 OverrideResY ) override;
+	virtual void RestoreSystemResolution(const TSharedRef<SWindow> InWindow) override {}
 	virtual void ReleaseDynamicResource( const FSlateBrush& Brush ) override;
 	virtual bool GenerateDynamicImageResource(FName ResourceName, uint32 Width, uint32 Height, const TArray< uint8 >& Bytes) override;
 	virtual void LoadStyleResources(const ISlateStyle& Style) override;
diff --git a/Engine/Source/Developer/StandaloneRenderer/Private/Windows/D3D/SlateD3DRenderer.h b/Engine/Source/Developer/StandaloneRenderer/Private/Windows/D3D/SlateD3DRenderer.h
--- a/Engine/Source/Developer/StandaloneRenderer/Private/Windows/D3D/SlateD3DRenderer.h
+++ b/Engine/Source/Developer/StandaloneRenderer/Private/Windows/D3D/SlateD3DRenderer.h
@@ -52,6 +52,7 @@ public:
 	virtual void CreateViewport( const TSharedRef<SWindow> InWindow ) override;
 	virtual void RequestResize( const TSharedPtr<SWindow>& InWindow, uint32 NewSizeX, uint32 NewSizeY ) override;
 	virtual void UpdateFullscreenState( const TSharedRef<SWindow> InWindow, uint32 OverrideResX, uint32 OverrideResY ) override;
+	virtual void RestoreSystemResolution(const TSharedRef<SWindow> InWindow) override {}
 	virtual void ReleaseDynamicResource( const FSlateBrush& Brush ) override;
 	virtual bool GenerateDynamicImageResource(FName ResourceName, uint32 Width, uint32 Height, const TArray< uint8 >& Bytes) override;
 	virtual void LoadStyleResources( const ISlateStyle& Style ) override;
