Linux: Handle errors creating OpenGL context earlier.

diff --git a/Engine/Source/Developer/StandaloneRenderer/Private/Linux/OpenGL/SlateOpenGLContext.cpp b/Engine/Source/Developer/StandaloneRenderer/Private/Linux/OpenGL/SlateOpenGLContext.cpp
--- a/Engine/Source/Developer/StandaloneRenderer/Private/Linux/OpenGL/SlateOpenGLContext.cpp
+++ b/Engine/Source/Developer/StandaloneRenderer/Private/Linux/OpenGL/SlateOpenGLContext.cpp
@@ -58,6 +58,23 @@ void FSlateOpenGLContext::Initialize( void* InWindow, const FSlateOpenGLContext*
 	}
 
 	Context = SDL_GL_CreateContext( WindowHandle );
+	if (Context == nullptr)
+	{
+		FString SdlError(ANSI_TO_TCHAR(SDL_GetError()));
+		
+		// ignore errors getting version, it will be clear from the logs
+		int OpenGLMajorVersion = -1;
+		int OpenGLMinorVersion = -1;
+		SDL_GL_GetAttribute( SDL_GL_CONTEXT_MAJOR_VERSION, &OpenGLMajorVersion );
+		SDL_GL_GetAttribute( SDL_GL_CONTEXT_MINOR_VERSION, &OpenGLMinorVersion );
+		
+		UE_LOG(LogInit, Fatal, TEXT("FSlateOpenGLContext::Initialize - Could not create OpenGL %d.%d context, SDL error: '%s'"), 
+				OpenGLMajorVersion, OpenGLMinorVersion,
+				*SdlError
+			);
+		// unreachable
+		return;
+	}
 	SDL_GL_MakeCurrent( WindowHandle, Context );
 }
 
diff --git a/Engine/Source/Runtime/OpenGLDrv/Private/Linux/OpenGLLinux.cpp b/Engine/Source/Runtime/OpenGLDrv/Private/Linux/OpenGLLinux.cpp
--- a/Engine/Source/Runtime/OpenGLDrv/Private/Linux/OpenGLLinux.cpp
+++ b/Engine/Source/Runtime/OpenGLDrv/Private/Linux/OpenGLLinux.cpp
@@ -159,6 +159,23 @@ static void _PlatformCreateOpenGLContextCore( FPlatformOpenGLContext* OutContext
 	OutContext->ViewportFramebuffer = 0;
 
 	OutContext->hGLContext = SDL_GL_CreateContext( OutContext->hWnd );
+	if (OutContext->hGLContext == nullptr)
+	{
+		FString SdlError(ANSI_TO_TCHAR(SDL_GetError()));
+		
+		// ignore errors getting version, it will be clear from the logs
+		int OpenGLMajorVersion = -1;
+		int OpenGLMinorVersion = -1;
+		SDL_GL_GetAttribute( SDL_GL_CONTEXT_MAJOR_VERSION, &OpenGLMajorVersion );
+		SDL_GL_GetAttribute( SDL_GL_CONTEXT_MINOR_VERSION, &OpenGLMinorVersion );
+		
+		UE_LOG(LogInit, Fatal, TEXT("_PlatformCreateOpenGLContextCore - Could not create OpenGL %d.%d context, SDL error: '%s'"), 
+				OpenGLMajorVersion, OpenGLMinorVersion,
+				*SdlError
+			);
+		// unreachable
+		return;
+	}
 
 	SDL_GL_MakeCurrent( prevWindow, prevContext );
 }
