Fixed a bug in Mac FSlateOpenGLContext::Initialize() that was making window height incorrect

diff --git a/Engine/Source/Developer/StandaloneRenderer/Private/Mac/OpenGL/SlateOpenGLContext.cpp b/Engine/Source/Developer/StandaloneRenderer/Private/Mac/OpenGL/SlateOpenGLContext.cpp
--- a/Engine/Source/Developer/StandaloneRenderer/Private/Mac/OpenGL/SlateOpenGLContext.cpp
+++ b/Engine/Source/Developer/StandaloneRenderer/Private/Mac/OpenGL/SlateOpenGLContext.cpp
@@ -189,11 +189,6 @@ void FSlateOpenGLContext::Initialize( void* InWindow, const FSlateOpenGLContext*
 			[View setAutoresizingMask:(NSViewWidthSizable | NSViewHeightSizable)];
 			
 			[SuperView addSubview:View];
-			
-			NSRect NewRect = [Window frame];
-			NewRect.size.height -= 22;
-			
-			[Window setFrame:NewRect display:YES];
 		}
 		else
 		{
