- Missed SlateViewer on IOS checking

diff --git a/Engine/Source/Developer/StandaloneRenderer/Public/IOS/SlateOpenGLESView.h b/Engine/Source/Developer/StandaloneRenderer/Public/IOS/SlateOpenGLESView.h
--- a/Engine/Source/Developer/StandaloneRenderer/Public/IOS/SlateOpenGLESView.h
+++ b/Engine/Source/Developer/StandaloneRenderer/Public/IOS/SlateOpenGLESView.h
@@ -8,6 +8,8 @@
 
 @interface SlateOpenGLESView : GLKView
 
+@property(retain) EAGLContext* Context;
+
 @end
 
 @interface SlateOpenGLESViewController : UIViewController
