Merging CL #2399881 from 4.7 to Main (Update OpenGL context on GPU switch in Mac standalone renderer)

diff --git a/Engine/Source/Developer/StandaloneRenderer/Private/Mac/OpenGL/SlateOpenGLContext.cpp b/Engine/Source/Developer/StandaloneRenderer/Private/Mac/OpenGL/SlateOpenGLContext.cpp
--- a/Engine/Source/Developer/StandaloneRenderer/Private/Mac/OpenGL/SlateOpenGLContext.cpp
+++ b/Engine/Source/Developer/StandaloneRenderer/Private/Mac/OpenGL/SlateOpenGLContext.cpp
@@ -146,6 +146,14 @@ char const* const CompositedBlitFragmentShader = "#version 120\n"
 
 @end
 
+static void MacOpenGLContextReconfigurationCallBack(CGDirectDisplayID Display, CGDisplayChangeSummaryFlags Flags, void* UserInfo)
+{
+	FSlateOpenGLContext* Context = (FSlateOpenGLContext*)UserInfo;
+	if (Context)
+	{
+		Context->bNeedsUpdate = true;
+	}
+}
 
 FSlateOpenGLContext::FSlateOpenGLContext()
 :	View(NULL)
@@ -158,6 +166,7 @@ FSlateOpenGLContext::FSlateOpenGLContext()
 ,	TextureDirectionUniform(0)
 ,	CompositeTexture(0)
 ,	CompositeVAO(0)
+,	bNeedsUpdate(false)
 {
 }
 
@@ -294,6 +303,8 @@ void FSlateOpenGLContext::Initialize( void* InWindow, const FSlateOpenGLContext*
 			glBindTexture (GL_TEXTURE_2D, 0);
 		}
 	}
+
+	CGDisplayRegisterReconfigurationCallback(&MacOpenGLContextReconfigurationCallBack, this);
 }
 
 void FSlateOpenGLContext::Destroy()
@@ -327,6 +338,8 @@ void FSlateOpenGLContext::Destroy()
 		[View release];
 		View = NULL;
 
+		CGDisplayRemoveReconfigurationCallback(&MacOpenGLContextReconfigurationCallBack, this);
+
 		// PixelFormat and Context are released by View
 		PixelFormat = NULL;
 		Context = NULL;
@@ -335,5 +348,10 @@ void FSlateOpenGLContext::Destroy()
 
 void FSlateOpenGLContext::MakeCurrent()
 {
+	if (bNeedsUpdate)
+	{
+		[Context update];
+		bNeedsUpdate = false;
+	}
 	[Context makeCurrentContext];
 }
diff --git a/Engine/Source/Developer/StandaloneRenderer/Private/OpenGL/SlateOpenGLRenderer.h b/Engine/Source/Developer/StandaloneRenderer/Private/OpenGL/SlateOpenGLRenderer.h
--- a/Engine/Source/Developer/StandaloneRenderer/Private/OpenGL/SlateOpenGLRenderer.h
+++ b/Engine/Source/Developer/StandaloneRenderer/Private/OpenGL/SlateOpenGLRenderer.h
@@ -38,6 +38,7 @@ struct FSlateOpenGLContext
 	GLint TextureDirectionUniform;
 	GLuint CompositeTexture;
 	GLuint CompositeVAO;
+	bool bNeedsUpdate;
 #elif PLATFORM_IOS
 	UIWindow* WindowHandle;
 	EAGLContext* Context;
