Merging using UE4-To-UE4-LauncherReleases REVERSE

CL# 2350227 on 05/11/2014 11:58:34
----------------------------------------------
Adding HRESULT logging to assertions on D3D11 calls in the Slate StandaloneRenderer

diff --git a/Engine/Source/Developer/StandaloneRenderer/Private/Windows/D3D/SlateD3DConstantBuffer.h b/Engine/Source/Developer/StandaloneRenderer/Private/Windows/D3D/SlateD3DConstantBuffer.h
--- a/Engine/Source/Developer/StandaloneRenderer/Private/Windows/D3D/SlateD3DConstantBuffer.h
+++ b/Engine/Source/Developer/StandaloneRenderer/Private/Windows/D3D/SlateD3DConstantBuffer.h
@@ -20,7 +20,7 @@ public:
 		Desc.StructureByteStride = 0;
 
 		HRESULT Hr = GD3DDevice->CreateBuffer(&Desc, NULL, Buffer.GetInitReference() );
-		check( SUCCEEDED(Hr) );
+		checkf( SUCCEEDED(Hr), TEXT("D3D11 Error Result %X"), Hr );
 	}
 
 	TRefCountPtr<ID3D11Buffer> GetResource() { return Buffer; }
diff --git a/Engine/Source/Developer/StandaloneRenderer/Private/Windows/D3D/SlateD3DIndexBuffer.cpp b/Engine/Source/Developer/StandaloneRenderer/Private/Windows/D3D/SlateD3DIndexBuffer.cpp
--- a/Engine/Source/Developer/StandaloneRenderer/Private/Windows/D3D/SlateD3DIndexBuffer.cpp
+++ b/Engine/Source/Developer/StandaloneRenderer/Private/Windows/D3D/SlateD3DIndexBuffer.cpp
@@ -35,7 +35,7 @@ void FSlateD3DIndexBuffer::CreateBuffer()
 	BufferDesc.MiscFlags = 0;
 
 	HRESULT Hr = GD3DDevice->CreateBuffer( &BufferDesc, NULL, Buffer.GetInitReference() ) ;
-	check( SUCCEEDED(Hr) );
+	checkf( SUCCEEDED(Hr), TEXT("D3D11 Error Result %X"), Hr );
 }
 
 /** Resizes the buffer to the passed in size.  Preserves internal data */
@@ -63,7 +63,7 @@ void FSlateD3DIndexBuffer::ResizeBuffer( uint32 NumIndices )
 		BufferDesc.MiscFlags = 0;
 
 		HRESULT Hr = GD3DDevice->CreateBuffer( &BufferDesc, NULL, Buffer.GetInitReference() ) ;
-		check( SUCCEEDED(Hr) );
+		checkf( SUCCEEDED(Hr), TEXT("D3D11 Error Result %X"), Hr );
 
 		if( SavedIndices )
 		{
diff --git a/Engine/Source/Developer/StandaloneRenderer/Private/Windows/D3D/SlateD3DRenderer.cpp b/Engine/Source/Developer/StandaloneRenderer/Private/Windows/D3D/SlateD3DRenderer.cpp
--- a/Engine/Source/Developer/StandaloneRenderer/Private/Windows/D3D/SlateD3DRenderer.cpp
+++ b/Engine/Source/Developer/StandaloneRenderer/Private/Windows/D3D/SlateD3DRenderer.cpp
@@ -109,7 +109,7 @@ void FSlateD3DRenderer::CreateDevice()
 		const D3D_FEATURE_LEVEL FeatureLevels[] = { D3D_FEATURE_LEVEL_11_0, D3D_FEATURE_LEVEL_10_1, D3D_FEATURE_LEVEL_10_0, D3D_FEATURE_LEVEL_9_3 };
 		D3D_FEATURE_LEVEL CreatedFeatureLevel;
 		HRESULT Hr = D3D11CreateDevice( NULL, DriverType, NULL, DeviceCreationFlags, FeatureLevels, sizeof(FeatureLevels)/sizeof(D3D_FEATURE_LEVEL), D3D11_SDK_VERSION, GD3DDevice.GetInitReference(), &CreatedFeatureLevel, GD3DDeviceContext.GetInitReference() );
-		check( SUCCEEDED(Hr) );
+		checkf( SUCCEEDED(Hr), TEXT("D3D11 Error Result %X"), Hr );
 	}
 }
 
@@ -174,23 +174,23 @@ void FSlateD3DRenderer::Private_CreateViewport( TSharedRef<SWindow> InWindow, co
 
 	TRefCountPtr<IDXGIDevice> DXGIDevice;
 	HRESULT Hr = GD3DDevice->QueryInterface( __uuidof(IDXGIDevice), (void**)DXGIDevice.GetInitReference() );
-	check( SUCCEEDED(Hr) );
+	checkf( SUCCEEDED(Hr), TEXT("D3D11 Error Result %X"), Hr );
 
 	TRefCountPtr<IDXGIAdapter> DXGIAdapter;
 	Hr = DXGIDevice->GetParent(__uuidof(IDXGIAdapter), (void **)DXGIAdapter.GetInitReference() );
-	check( SUCCEEDED(Hr) );
+	checkf( SUCCEEDED(Hr), TEXT("D3D11 Error Result %X"), Hr );
 
 	TRefCountPtr<IDXGIFactory> DXGIFactory;
-	DXGIAdapter->GetParent(__uuidof(IDXGIFactory), (void **)DXGIFactory.GetInitReference() );
-	check( SUCCEEDED(Hr) );
+	Hr = DXGIAdapter->GetParent(__uuidof(IDXGIFactory), (void **)DXGIFactory.GetInitReference());
+	checkf( SUCCEEDED(Hr), TEXT("D3D11 Error Result %X"), Hr );
 
 	FSlateD3DViewport Viewport;
 
 	Hr = DXGIFactory->CreateSwapChain(DXGIDevice.GetReference(), &SwapChainDesc, Viewport.D3DSwapChain.GetInitReference() );
-	check( SUCCEEDED(Hr) );
+	checkf( SUCCEEDED(Hr), TEXT("D3D11 Error Result %X"), Hr );
 
 	Hr = DXGIFactory->MakeWindowAssociation((HWND)NativeWindow->GetOSWindowHandle(),DXGI_MWA_NO_ALT_ENTER);
-	check(SUCCEEDED(Hr));
+	checkf( SUCCEEDED(Hr), TEXT("D3D11 Error Result %X"), Hr );
 
 	uint32 Width = FMath::TruncToInt(WindowSize.X);
 	uint32 Height = FMath::TruncToInt(WindowSize.Y);
@@ -253,7 +253,7 @@ void FSlateD3DRenderer::Private_ResizeViewport( const TSharedRef<SWindow> InWind
 		DXGI_SWAP_CHAIN_DESC Desc;
 		Viewport->D3DSwapChain->GetDesc( &Desc );
 		HRESULT Hr = Viewport->D3DSwapChain->ResizeBuffers( Desc.BufferCount, Viewport->ViewportInfo.Width, Viewport->ViewportInfo.Height, Desc.BufferDesc.Format, Desc.Flags );
-		check( SUCCEEDED(Hr) );
+		checkf( SUCCEEDED(Hr), TEXT("D3D11 Error Result %X"), Hr );
 
 		CreateBackBufferResources( Viewport->D3DSwapChain, Viewport->BackBufferTexture,Viewport->RenderTargetView );
 	}
@@ -268,7 +268,7 @@ void FSlateD3DRenderer::CreateBackBufferResources( TRefCountPtr<IDXGISwapChain>&
 	RTVDesc.ViewDimension = D3D11_RTV_DIMENSION_TEXTURE2D;
 	RTVDesc.Texture2D.MipSlice = 0;
 	HRESULT Hr = GD3DDevice->CreateRenderTargetView( OutBackBuffer, &RTVDesc, OutRTV.GetInitReference() );
-	check( SUCCEEDED(Hr) );
+	checkf( SUCCEEDED(Hr), TEXT("D3D11 Error Result %X"), Hr );
 }
 
 void FSlateD3DRenderer::CreateViewport( const TSharedRef<SWindow> InWindow )
@@ -305,7 +305,7 @@ void FSlateD3DRenderer::CreateDepthStencilBuffer( FSlateD3DViewport& Viewport )
 	DescDepth.CPUAccessFlags = 0;
 	DescDepth.MiscFlags = 0;
 	HRESULT Hr = GD3DDevice->CreateTexture2D( &DescDepth, NULL, DepthStencil.GetInitReference() );
-	check( SUCCEEDED(Hr) );
+	checkf( SUCCEEDED(Hr), TEXT("D3D11 Error Result %X"), Hr );
 
 	D3D11_DEPTH_STENCIL_VIEW_DESC DescDSV;
 #if DEPTH_32_BIT_CONVERSION
@@ -319,7 +319,7 @@ void FSlateD3DRenderer::CreateDepthStencilBuffer( FSlateD3DViewport& Viewport )
 
 	// Create the depth stencil view
 	Hr = GD3DDevice->CreateDepthStencilView( DepthStencil, &DescDSV, Viewport.DepthStencilView.GetInitReference() ); 
-	check( SUCCEEDED(Hr) );
+	checkf( SUCCEEDED(Hr), TEXT("D3D11 Error Result %X"), Hr );
 }
 
 void FSlateD3DRenderer::DrawWindows( FSlateDrawBuffer& InWindowDrawBuffer )
diff --git a/Engine/Source/Developer/StandaloneRenderer/Private/Windows/D3D/SlateD3DRenderingPolicy.cpp b/Engine/Source/Developer/StandaloneRenderer/Private/Windows/D3D/SlateD3DRenderingPolicy.cpp
--- a/Engine/Source/Developer/StandaloneRenderer/Private/Windows/D3D/SlateD3DRenderingPolicy.cpp
+++ b/Engine/Source/Developer/StandaloneRenderer/Private/Windows/D3D/SlateD3DRenderingPolicy.cpp
@@ -54,22 +54,22 @@ void FSlateD3D11RenderingPolicy::InitResources()
 	SamplerDesc.MinLOD = 0;
 
 	HRESULT Hr = GD3DDevice->CreateSamplerState( &SamplerDesc, PointSamplerState_Wrap.GetInitReference() );
-	check( SUCCEEDED( Hr ) );
+	checkf( SUCCEEDED(Hr), TEXT("D3D11 Error Result %X"), Hr );
 
 	SamplerDesc.Filter = D3D11_FILTER_MIN_MAG_LINEAR_MIP_POINT;
 	Hr = GD3DDevice->CreateSamplerState( &SamplerDesc, BilinearSamplerState_Wrap.GetInitReference() );
-	check( SUCCEEDED( Hr ) );
+	checkf( SUCCEEDED(Hr), TEXT("D3D11 Error Result %X"), Hr );
 
 	SamplerDesc.AddressU = D3D11_TEXTURE_ADDRESS_CLAMP;
 	SamplerDesc.AddressV = D3D11_TEXTURE_ADDRESS_CLAMP;
 	SamplerDesc.AddressW = D3D11_TEXTURE_ADDRESS_CLAMP;
 
 	Hr = GD3DDevice->CreateSamplerState( &SamplerDesc, BilinearSamplerState_Clamp.GetInitReference() );
-	check( SUCCEEDED( Hr ) );
+	checkf( SUCCEEDED(Hr), TEXT("D3D11 Error Result %X"), Hr );
 
 	SamplerDesc.Filter = D3D11_FILTER_MIN_MAG_MIP_POINT;
 	Hr = GD3DDevice->CreateSamplerState( &SamplerDesc, PointSamplerState_Clamp.GetInitReference() );
-	check( SUCCEEDED( Hr ) );
+	checkf( SUCCEEDED(Hr), TEXT("D3D11 Error Result %X"), Hr );
 
 	WhiteTexture = TextureManager->CreateColorTexture( TEXT("DefaultWhite"), FColor::White )->Resource;
 
diff --git a/Engine/Source/Developer/StandaloneRenderer/Private/Windows/D3D/SlateD3DShaders.cpp b/Engine/Source/Developer/StandaloneRenderer/Private/Windows/D3D/SlateD3DShaders.cpp
--- a/Engine/Source/Developer/StandaloneRenderer/Private/Windows/D3D/SlateD3DShaders.cpp
+++ b/Engine/Source/Developer/StandaloneRenderer/Private/Windows/D3D/SlateD3DShaders.cpp
@@ -73,14 +73,14 @@ void FSlateD3DVS::Create( const FString& Filename, const FString& EntryPoint, co
 	CompileShader( Filename, EntryPoint, ShaderModel, Blob);
 
 	HRESULT Hr = GD3DDevice->CreateVertexShader( Blob->GetBufferPointer(), Blob->GetBufferSize(), NULL, VertexShader.GetInitReference() );
-	check( SUCCEEDED(Hr) );
+	checkf( SUCCEEDED(Hr), TEXT("D3D11 Error Result %X"), Hr );
 
 	Hr = GD3DDevice->CreateInputLayout( VertexLayout, VertexLayoutCount, Blob->GetBufferPointer(), Blob->GetBufferSize(), InputLayout.GetInitReference() );
-	check( SUCCEEDED(Hr) );
+	checkf( SUCCEEDED(Hr), TEXT("D3D11 Error Result %X"), Hr );
 
 	TRefCountPtr<ID3D11ShaderReflection> Reflector;
 	Hr = D3DReflect( Blob->GetBufferPointer(), Blob->GetBufferSize(), IID_ID3D11ShaderReflection, (void**)Reflector.GetInitReference() );
-	check( SUCCEEDED(Hr) );
+	checkf( SUCCEEDED(Hr), TEXT("D3D11 Error Result %X"), Hr );
 
 	GetShaderBindings( Reflector, ShaderBindings );
 }
@@ -129,12 +129,12 @@ void FSlateD3DGeometryShader::Create( const FString& Filename, const FString& En
 	CompileShader( Filename, EntryPoint, ShaderModel, Blob);
 
 	HRESULT Hr = GD3DDevice->CreateGeometryShader( Blob->GetBufferPointer(), Blob->GetBufferSize(), NULL, GeometryShader.GetInitReference() );
-	check( SUCCEEDED(Hr) );
+	checkf( SUCCEEDED(Hr), TEXT("D3D11 Error Result %X"), Hr );
 
 
 	TRefCountPtr<ID3D11ShaderReflection> Reflector;
 	Hr = D3DReflect( Blob->GetBufferPointer(), Blob->GetBufferSize(), IID_ID3D11ShaderReflection, (void**)Reflector.GetInitReference() );
-	check( SUCCEEDED(Hr) );
+	checkf( SUCCEEDED(Hr), TEXT("D3D11 Error Result %X"), Hr );
 
 	GetShaderBindings( Reflector, ShaderBindings );
 }
@@ -170,12 +170,12 @@ void FSlateD3DPS::Create( const FString& Filename, const FString& EntryPoint, co
 	CompileShader( Filename, EntryPoint, ShaderModel, Blob);
 
 	HRESULT Hr = GD3DDevice->CreatePixelShader( Blob->GetBufferPointer(), Blob->GetBufferSize(), NULL, PixelShader.GetInitReference() );
-	check( SUCCEEDED(Hr) );
+	checkf( SUCCEEDED(Hr), TEXT("D3D11 Error Result %X"), Hr );
 
 
 	TRefCountPtr<ID3D11ShaderReflection> Reflector;
 	Hr = D3DReflect( Blob->GetBufferPointer(), Blob->GetBufferSize(), IID_ID3D11ShaderReflection, (void**)Reflector.GetInitReference() );
-	check( SUCCEEDED(Hr) );
+	checkf( SUCCEEDED(Hr), TEXT("D3D11 Error Result %X"), Hr );
 
 	GetShaderBindings( Reflector, ShaderBindings );
 }
diff --git a/Engine/Source/Developer/StandaloneRenderer/Private/Windows/D3D/SlateD3DTextures.cpp b/Engine/Source/Developer/StandaloneRenderer/Private/Windows/D3D/SlateD3DTextures.cpp
--- a/Engine/Source/Developer/StandaloneRenderer/Private/Windows/D3D/SlateD3DTextures.cpp
+++ b/Engine/Source/Developer/StandaloneRenderer/Private/Windows/D3D/SlateD3DTextures.cpp
@@ -21,7 +21,7 @@ void FSlateD3DTexture::Init( DXGI_FORMAT InFormat, D3D11_SUBRESOURCE_DATA* Inita
 	TexDesc.MiscFlags = 0;
 
 	HRESULT Hr = GD3DDevice->CreateTexture2D( &TexDesc, InitalData, D3DTexture.GetInitReference() );
-	check( SUCCEEDED(Hr) );
+	checkf( SUCCEEDED(Hr), TEXT("D3D11 Error Result %X"), Hr );
 
 	// Create the shader accessable view of the texture
 	D3D11_SHADER_RESOURCE_VIEW_DESC SrvDesc;
@@ -32,7 +32,7 @@ void FSlateD3DTexture::Init( DXGI_FORMAT InFormat, D3D11_SUBRESOURCE_DATA* Inita
 
 	// Create the shader resource view.
 	Hr = GD3DDevice->CreateShaderResourceView( D3DTexture, &SrvDesc, ShaderResource.GetInitReference() );
-	check( SUCCEEDED(Hr) );
+	checkf( SUCCEEDED(Hr), TEXT("D3D11 Error Result %X"), Hr );
 }
 
 void FSlateD3DTexture::ResizeTexture(uint32 Width, uint32 Height)
diff --git a/Engine/Source/Developer/StandaloneRenderer/Private/Windows/D3D/SlateD3DVertexBuffer.cpp b/Engine/Source/Developer/StandaloneRenderer/Private/Windows/D3D/SlateD3DVertexBuffer.cpp
--- a/Engine/Source/Developer/StandaloneRenderer/Private/Windows/D3D/SlateD3DVertexBuffer.cpp
+++ b/Engine/Source/Developer/StandaloneRenderer/Private/Windows/D3D/SlateD3DVertexBuffer.cpp
@@ -37,7 +37,7 @@ void FSlateD3DVertexBuffer::CreateBuffer( uint32 InStride )
 	BufferDesc.MiscFlags = 0;
 
 	HRESULT Hr = GD3DDevice->CreateBuffer( &BufferDesc, NULL, Buffer.GetInitReference() );
-	check( SUCCEEDED(Hr) );
+	checkf( SUCCEEDED(Hr), TEXT("D3D11 Error Result %X"), Hr );
 }
 
 /** Releases the vertex buffers RHI resource. */
@@ -73,7 +73,7 @@ void FSlateD3DVertexBuffer::ResizeBuffer( uint32 NewSize )
 		BufferDesc.MiscFlags = 0;
 
 		HRESULT Hr = GD3DDevice->CreateBuffer( &BufferDesc, NULL, Buffer.GetInitReference() );
-		check( SUCCEEDED(Hr) );
+		checkf( SUCCEEDED(Hr), TEXT("D3D11 Error Result %X"), Hr );
 
 		if( SavedVertices )
 		{
