Fixed crash in FSlateOpenGLRenderer::CreateUpdatableTexture()

diff --git a/Engine/Source/Developer/StandaloneRenderer/Private/OpenGL/SlateOpenGLRenderer.cpp b/Engine/Source/Developer/StandaloneRenderer/Private/OpenGL/SlateOpenGLRenderer.cpp
--- a/Engine/Source/Developer/StandaloneRenderer/Private/OpenGL/SlateOpenGLRenderer.cpp
+++ b/Engine/Source/Developer/StandaloneRenderer/Private/OpenGL/SlateOpenGLRenderer.cpp
@@ -227,7 +227,7 @@ void FSlateOpenGLRenderer::LoadStyleResources( const ISlateStyle& Style )
 FSlateUpdatableTexture* FSlateOpenGLRenderer::CreateUpdatableTexture(uint32 Width, uint32 Height)
 {
 	TArray<uint8> RawData;
-	RawData.AddZeroed(4);
+	RawData.AddZeroed(Width * Height * 4);
 	FSlateOpenGLTexture* NewTexture = new FSlateOpenGLTexture(Width, Height);
 #if !PLATFORM_USES_ES2
 	NewTexture->Init(GL_SRGB8_ALPHA8, RawData);
