OpenGL: Fix missing text issue in standalone programs.

- Set shared GL context after destroying a window as FontCache assumes a valid current context.

diff --git a/Engine/Source/Developer/StandaloneRenderer/Private/OpenGL/SlateOpenGLRenderer.cpp b/Engine/Source/Developer/StandaloneRenderer/Private/OpenGL/SlateOpenGLRenderer.cpp
--- a/Engine/Source/Developer/StandaloneRenderer/Private/OpenGL/SlateOpenGLRenderer.cpp
+++ b/Engine/Source/Developer/StandaloneRenderer/Private/OpenGL/SlateOpenGLRenderer.cpp
@@ -157,6 +157,8 @@ void FSlateOpenGLRenderer::OnWindowDestroyed( const TSharedRef<SWindow>& InWindo
 		Viewport->Destroy();
 	}
 	WindowToViewportMap.Remove( &InWindow.Get() );
+
+	SharedContext.MakeCurrent();
 }
 
 void FSlateOpenGLRenderer::CreateViewport( const TSharedRef<SWindow> InWindow )
