OpenGL/Linux: minor cleanup in standalone renderer.

- Fewer differences between code paths and clearer name for #define.
- Also removes GL error that was happening due to GL_ALPHA being used with a 3.2 context.

diff --git a/Engine/Source/Developer/StandaloneRenderer/Private/Linux/OpenGL/SlateOpenGLContext.cpp b/Engine/Source/Developer/StandaloneRenderer/Private/Linux/OpenGL/SlateOpenGLContext.cpp
--- a/Engine/Source/Developer/StandaloneRenderer/Private/Linux/OpenGL/SlateOpenGLContext.cpp
+++ b/Engine/Source/Developer/StandaloneRenderer/Private/Linux/OpenGL/SlateOpenGLContext.cpp
@@ -116,7 +116,6 @@ void FSlateOpenGLContext::MakeCurrent()
 		if(SDL_GL_MakeCurrent( WindowHandle, Context ) == 0)
 		{
 			glBindVertexArray(VertexArrayObject);
-			glGetError(); // SDL leaves glGetError in a dirty state even when successful?
 		}
 	}
 }
diff --git a/Engine/Source/Developer/StandaloneRenderer/Private/OpenGL/SlateOpenGLTextures.cpp b/Engine/Source/Developer/StandaloneRenderer/Private/OpenGL/SlateOpenGLTextures.cpp
--- a/Engine/Source/Developer/StandaloneRenderer/Private/OpenGL/SlateOpenGLTextures.cpp
+++ b/Engine/Source/Developer/StandaloneRenderer/Private/OpenGL/SlateOpenGLTextures.cpp
@@ -4,6 +4,8 @@
 #include "SlateOpenGLRenderer.h"
 #include "OpenGL/SlateOpenGLTextures.h"
 
+#define USE_DEPRECATED_OPENGL_FUNCTIONALITY			(!PLATFORM_USES_ES2 && !PLATFORM_LINUX)
+
 GLuint FSlateOpenGLTexture::NullTexture = 0;
 
 void FSlateOpenGLTexture::Init( GLenum TexFormat, const TArray<uint8>& TextureData )
@@ -13,15 +15,15 @@ void FSlateOpenGLTexture::Init( GLenum TexFormat, const TArray<uint8>& TextureDa
 	CHECK_GL_ERRORS;
 
 	// Ensure texturing is enabled before setting texture properties
-#if !PLATFORM_USES_ES2 && !PLATFORM_LINUX
-	glEnable(GL_TEXTURE_2D);
-#endif
+#if USE_DEPRECATED_OPENGL_FUNCTIONALITY
+    glEnable(GL_TEXTURE_2D);
+#endif // USE_DEPRECATED_OPENGL_FUNCTIONALITY
 	glBindTexture(GL_TEXTURE_2D, ShaderResource);
 	glTexParameteri(GL_TEXTURE_2D, GL_TEXTURE_MIN_FILTER, GL_LINEAR);
 	glTexParameteri(GL_TEXTURE_2D, GL_TEXTURE_MAG_FILTER, GL_LINEAR);
-#if !PLATFORM_LINUX
+#if USE_DEPRECATED_OPENGL_FUNCTIONALITY
 	glTexEnvi( GL_TEXTURE_ENV, GL_TEXTURE_ENV_MODE, GL_REPLACE );
-#endif
+#endif // USE_DEPRECATED_OPENGL_FUNCTIONALITY
 
 	// the raw data is in bgra or bgr
 	const GLint Format = GL_BGRA;
@@ -45,15 +47,15 @@ void FSlateOpenGLTexture::ResizeTexture(uint32 Width, uint32 Height)
 void FSlateOpenGLTexture::UpdateTexture(const TArray<uint8>& Bytes)
 {
 	// Ensure texturing is enabled before setting texture properties
-#if !PLATFORM_USES_ES2 && !PLATFORM_LINUX
-	glEnable(GL_TEXTURE_2D);
-#endif
+#if USE_DEPRECATED_OPENGL_FUNCTIONALITY
+    glEnable(GL_TEXTURE_2D);
+#endif // USE_DEPRECATED_OPENGL_FUNCTIONALITY
 	glBindTexture(GL_TEXTURE_2D, ShaderResource);
 	glTexParameteri(GL_TEXTURE_2D, GL_TEXTURE_MIN_FILTER, GL_LINEAR);
 	glTexParameteri(GL_TEXTURE_2D, GL_TEXTURE_MAG_FILTER, GL_LINEAR);
-#if !PLATFORM_LINUX
-	glTexEnvi(GL_TEXTURE_ENV, GL_TEXTURE_ENV_MODE, GL_REPLACE);
-#endif
+#if USE_DEPRECATED_OPENGL_FUNCTIONALITY
+    glTexEnvi(GL_TEXTURE_ENV, GL_TEXTURE_ENV_MODE, GL_REPLACE);
+#endif // USE_DEPRECATED_OPENGL_FUNCTIONALITY
 
 	// Upload the texture data
 #if !PLATFORM_USES_ES2
@@ -88,8 +90,15 @@ void FSlateFontTextureOpenGL::CreateFontTexture()
 	glTexParameteri(GL_TEXTURE_2D, GL_TEXTURE_MAG_FILTER, GL_NEAREST);
 
 	// Font textures use an alpha texture only
-	GLint Format = GL_ALPHA;
-	// Upload the data to the texture
+    GLint Format =
+#if USE_DEPRECATED_OPENGL_FUNCTIONALITY
+        GL_ALPHA
+#else
+		GL_RED
+#endif // USE_DEPRECATED_OPENGL_FUNCTIONALITY
+    ;
+
+    // Upload the data to the texture
 	glTexImage2D( GL_TEXTURE_2D, 0, Format, AtlasWidth, AtlasHeight, 0, Format, GL_UNSIGNED_BYTE, NULL );
 
 	// Create a new slate texture for use in rendering
@@ -108,7 +117,14 @@ void FSlateFontTextureOpenGL::ConditionalUpdateTexture()
 
 		// Completely the texture data each time characters are added
 		glBindTexture(GL_TEXTURE_2D, FontTexture->GetTypedResource() );
-		GLint Format = GL_ALPHA;
+        GLint Format =
+#if USE_DEPRECATED_OPENGL_FUNCTIONALITY
+            GL_ALPHA
+#else
+            GL_RED
+#endif // USE_DEPRECATED_OPENGL_FUNCTIONALITY
+        ;
+
 #if PLATFORM_MAC // Make this texture use a DMA'd client storage backing store on OS X, where these extensions always exist
 				 // This avoids a problem on Intel & Nvidia cards that makes characters disappear as well as making the texture updates
 				 // as fast as they possibly can be.
@@ -116,11 +132,7 @@ void FSlateFontTextureOpenGL::ConditionalUpdateTexture()
 		glTexParameteri(GL_TEXTURE_2D, GL_TEXTURE_STORAGE_HINT_APPLE, GL_STORAGE_CACHED_APPLE);
 		glPixelStorei(GL_UNPACK_CLIENT_STORAGE_APPLE, GL_TRUE);
 #endif
-#if PLATFORM_LINUX
-		glTexImage2D( GL_TEXTURE_2D, 0, GL_R8, AtlasWidth, AtlasHeight, 0, GL_RED, GL_UNSIGNED_BYTE, AtlasData.GetData() );
-#else
 		glTexImage2D( GL_TEXTURE_2D, 0, Format, AtlasWidth, AtlasHeight, 0, Format, GL_UNSIGNED_BYTE, AtlasData.GetData() );
-#endif
 #if PLATFORM_MAC
 		glPixelStorei(GL_UNPACK_CLIENT_STORAGE_APPLE, GL_FALSE);
 #endif
