Linux: use OpenGL 2.1 context in the standalone renderer.

- Should make programs (SlateViewer, UFE, Launcher) run on more drivers/hardware.

diff --git a/Engine/Source/Developer/StandaloneRenderer/Private/Linux/OpenGL/SlateOpenGLContext.cpp b/Engine/Source/Developer/StandaloneRenderer/Private/Linux/OpenGL/SlateOpenGLContext.cpp
--- a/Engine/Source/Developer/StandaloneRenderer/Private/Linux/OpenGL/SlateOpenGLContext.cpp
+++ b/Engine/Source/Developer/StandaloneRenderer/Private/Linux/OpenGL/SlateOpenGLContext.cpp
@@ -24,7 +24,9 @@ FSlateOpenGLContext::FSlateOpenGLContext()
 	: WindowHandle(NULL)
 	, Context(NULL)
 	, bReleaseWindowOnDestroy(false)
+#if LINUX_USE_OPENGL_3_2
 	, VertexArrayObject(0)
+#endif // LINUX_USE_OPENGL_3_2
 {
 }
 
@@ -49,8 +51,16 @@ void FSlateOpenGLContext::Initialize( void* InWindow, const FSlateOpenGLContext*
 		UE_LOG(LogInit, Fatal, TEXT("FSlateOpenGLContext::Initialize - Unable to dynamically load libGL: %s\n."), *SdlError);
 	}
 
-	SDL_GL_SetAttribute( SDL_GL_CONTEXT_MAJOR_VERSION, 3 );
-	SDL_GL_SetAttribute( SDL_GL_CONTEXT_MINOR_VERSION, 2 );
+	int OpenGLMajorVersionToUse = 2;
+	int OpenGLMinorVersionToUse = 1;
+
+#if LINUX_USE_OPENGL_3_2
+	OpenGLMajorVersionToUse = 3;
+	OpenGLMinorVersionToUse = 2;
+#endif // LINUX_USE_OPENGL_3_2
+
+	SDL_GL_SetAttribute( SDL_GL_CONTEXT_MAJOR_VERSION, OpenGLMajorVersionToUse );
+	SDL_GL_SetAttribute( SDL_GL_CONTEXT_MINOR_VERSION, OpenGLMinorVersionToUse );
 	SDL_GL_SetAttribute( SDL_GL_CONTEXT_PROFILE_MASK, SDL_GL_CONTEXT_PROFILE_CORE );
 	SDL_GL_SetAttribute( SDL_GL_CONTEXT_FLAGS, SDL_GL_CONTEXT_DEBUG_FLAG );
 
@@ -64,6 +74,7 @@ void FSlateOpenGLContext::Initialize( void* InWindow, const FSlateOpenGLContext*
 		SDL_GL_SetAttribute( SDL_GL_SHARE_WITH_CURRENT_CONTEXT, 0 );
 	}
 
+	UE_LOG(LogInit, Log, TEXT("FSlateOpenGLContext::Initialize - creating OpenGL %d.%d context"), OpenGLMajorVersionToUse, OpenGLMinorVersionToUse);
 	Context = SDL_GL_CreateContext( WindowHandle );
 	if (Context == nullptr)
 	{
@@ -86,9 +97,11 @@ void FSlateOpenGLContext::Initialize( void* InWindow, const FSlateOpenGLContext*
 
 	LoadOpenGLExtensions();
 
+#if LINUX_USE_OPENGL_3_2
 	// one Vertex Array Object is reportedly needed for OpenGL 3.2+ core profiles
 	glGenVertexArrays(1, &VertexArrayObject);
 	glBindVertexArray(VertexArrayObject);
+#endif // LINUX_USE_OPENGL_3_2
 }
 
 void FSlateOpenGLContext::Destroy()
@@ -96,7 +109,9 @@ void FSlateOpenGLContext::Destroy()
 	if	( WindowHandle != NULL )
 	{
 		SDL_GL_MakeCurrent( NULL, NULL );
+#if LINUX_USE_OPENGL_3_2
 		glDeleteVertexArrays(1, &VertexArrayObject);
+#endif // LINUX_USE_OPENGL_3_2
 		SDL_GL_DeleteContext( Context );
 
 		if	( bReleaseWindowOnDestroy )
@@ -115,7 +130,9 @@ void FSlateOpenGLContext::MakeCurrent()
 	{
 		if(SDL_GL_MakeCurrent( WindowHandle, Context ) == 0)
 		{
+#if LINUX_USE_OPENGL_3_2
 			glBindVertexArray(VertexArrayObject);
+#endif // LINUX_USE_OPENGL_3_2
 		}
 	}
 }
diff --git a/Engine/Source/Developer/StandaloneRenderer/Private/OpenGL/SlateOpenGLRenderer.h b/Engine/Source/Developer/StandaloneRenderer/Private/OpenGL/SlateOpenGLRenderer.h
--- a/Engine/Source/Developer/StandaloneRenderer/Private/OpenGL/SlateOpenGLRenderer.h
+++ b/Engine/Source/Developer/StandaloneRenderer/Private/OpenGL/SlateOpenGLRenderer.h
@@ -7,6 +7,10 @@ class FSlateOpenGLTextureManager;
 class FSlateOpenGLRenderingPolicy;
 class FSlateElementBatcher;
 
+// Optionally use 3.2 context on Linux. There's no real need to require this in a standalone application since it only renders Slate UI.
+// With this set to 0, StandaloneRenderer will use OpenGL 2.1 on Linux, which is almost universally supported (as of 2015).
+#define LINUX_USE_OPENGL_3_2		0
+
 #if !(UE_BUILD_SHIPPING || UE_BUILD_TEST)
 #define CHECK_GL_ERRORS \
 {\
@@ -41,7 +45,10 @@ struct FSlateOpenGLContext
 	SDL_Window* WindowHandle;
 	SDL_GLContext Context;
 	bool bReleaseWindowOnDestroy;
+#if LINUX_USE_OPENGL_3_2
 	GLuint VertexArrayObject; 
+#endif // LINUX_USE_OPENGL_3_2
+
 #else
 #error "Unknown platform"
 #endif
diff --git a/Engine/Source/Developer/StandaloneRenderer/Private/OpenGL/SlateOpenGLShaders.cpp b/Engine/Source/Developer/StandaloneRenderer/Private/OpenGL/SlateOpenGLShaders.cpp
--- a/Engine/Source/Developer/StandaloneRenderer/Private/OpenGL/SlateOpenGLShaders.cpp
+++ b/Engine/Source/Developer/StandaloneRenderer/Private/OpenGL/SlateOpenGLShaders.cpp
@@ -95,7 +95,11 @@ void FSlateOpenGLShader::CompileShader( const FString& Filename, GLenum ShaderTy
 #if PLATFORM_USES_ES2
 	Chars[0] = (ANSICHAR*)"#define PLATFORM_USES_ES2 1\n\n#define PLATFORM_LINUX 0\n";
 #elif PLATFORM_LINUX
+	#if LINUX_USE_OPENGL_3_2
 	Chars[0] = (ANSICHAR*)"#version 150\n\n#define PLATFORM_USES_ES2 0\n\n#define PLATFORM_LINUX 1\n";
+	#else
+	Chars[0] = (ANSICHAR*)"#version 120\n\n#define PLATFORM_USES_ES2 0\n\n#define PLATFORM_LINUX 1\n";
+	#endif // LINUX_USE_OPENGL_3_2
 #else
 	Chars[0] = (ANSICHAR*)"#version 120\n\n#define PLATFORM_USES_ES2 0\n\n#define PLATFORM_LINUX 0\n";
 #endif
