Merging using UE4-To-UE4-LauncherDev

diff --git a/Engine/Source/Developer/StandaloneRenderer/Private/StandaloneRenderer.cpp b/Engine/Source/Developer/StandaloneRenderer/Private/StandaloneRenderer.cpp
--- a/Engine/Source/Developer/StandaloneRenderer/Private/StandaloneRenderer.cpp
+++ b/Engine/Source/Developer/StandaloneRenderer/Private/StandaloneRenderer.cpp
@@ -8,6 +8,8 @@
 
 #include "OpenGL/SlateOpenGLRenderer.h"
 
+DEFINE_LOG_CATEGORY(LogStandaloneRenderer);
+
 /**
  * Single function to create the standalone renderer for the running platform
  */
diff --git a/Engine/Source/Developer/StandaloneRenderer/Private/StandaloneRendererPrivate.h b/Engine/Source/Developer/StandaloneRenderer/Private/StandaloneRendererPrivate.h
--- a/Engine/Source/Developer/StandaloneRenderer/Private/StandaloneRendererPrivate.h
+++ b/Engine/Source/Developer/StandaloneRenderer/Private/StandaloneRendererPrivate.h
@@ -3,6 +3,7 @@
 #ifndef UNREAL_STANDALONERENDERER_PRIVATEPCH_H
 #define UNREAL_STANDALONERENDERER_PRIVATEPCH_H
 
+DECLARE_LOG_CATEGORY_EXTERN(LogStandaloneRenderer, Log, All);
 
 #include "Core.h"
 #include "SlateCore.h"
diff --git a/Engine/Source/Developer/StandaloneRenderer/Private/Windows/D3D/SlateD3DRenderer.cpp b/Engine/Source/Developer/StandaloneRenderer/Private/Windows/D3D/SlateD3DRenderer.cpp
--- a/Engine/Source/Developer/StandaloneRenderer/Private/Windows/D3D/SlateD3DRenderer.cpp
+++ b/Engine/Source/Developer/StandaloneRenderer/Private/Windows/D3D/SlateD3DRenderer.cpp
@@ -114,7 +114,8 @@ void FSlateD3DRenderer::CreateDevice()
 		const D3D_FEATURE_LEVEL FeatureLevels[] = { D3D_FEATURE_LEVEL_11_0, D3D_FEATURE_LEVEL_10_1, D3D_FEATURE_LEVEL_10_0, D3D_FEATURE_LEVEL_9_3 };
 		D3D_FEATURE_LEVEL CreatedFeatureLevel;
 		HRESULT Hr = D3D11CreateDevice( NULL, DriverType, NULL, DeviceCreationFlags, FeatureLevels, sizeof(FeatureLevels)/sizeof(D3D_FEATURE_LEVEL), D3D11_SDK_VERSION, GD3DDevice.GetInitReference(), &CreatedFeatureLevel, GD3DDeviceContext.GetInitReference() );
-		checkf( SUCCEEDED(Hr), TEXT("D3D11 Error Result %X"), Hr );
+		UE_LOG(LogStandaloneRenderer, Log, TEXT("D3D11CreateDevice Result: %X"), Hr);
+		checkf(SUCCEEDED(Hr), TEXT("D3D11 Error Result %X"), Hr);
 	}
 }
 
