Add method to UpdatableTexture to enable passing in raw texture data from 3rd party APIs

Merged CL 2533543 + 2534462 from LauncherDev to Main using UE4-To-UE4-LauncherDev

diff --git a/Engine/Source/Developer/StandaloneRenderer/Private/OpenGL/SlateOpenGLTextures.cpp b/Engine/Source/Developer/StandaloneRenderer/Private/OpenGL/SlateOpenGLTextures.cpp
--- a/Engine/Source/Developer/StandaloneRenderer/Private/OpenGL/SlateOpenGLTextures.cpp
+++ b/Engine/Source/Developer/StandaloneRenderer/Private/OpenGL/SlateOpenGLTextures.cpp
@@ -16,7 +16,7 @@ void FSlateOpenGLTexture::Init( GLenum TexFormat, const TArray<uint8>& TextureDa
 
 	// Ensure texturing is enabled before setting texture properties
 #if USE_DEPRECATED_OPENGL_FUNCTIONALITY
-    glEnable(GL_TEXTURE_2D);
+	glEnable(GL_TEXTURE_2D);
 #endif // USE_DEPRECATED_OPENGL_FUNCTIONALITY
 	glBindTexture(GL_TEXTURE_2D, ShaderResource);
 	glTexParameteri(GL_TEXTURE_2D, GL_TEXTURE_MIN_FILTER, GL_LINEAR);
@@ -45,27 +45,41 @@ void FSlateOpenGLTexture::ResizeTexture(uint32 Width, uint32 Height)
 }
 
 void FSlateOpenGLTexture::UpdateTexture(const TArray<uint8>& Bytes)
+{
+	UpdateTextureRaw(Bytes.GetData());
+}
+
+void FSlateOpenGLTexture::UpdateTextureThreadSafeRaw(uint32 Width, uint32 Height, const void* Buffer)
+{
+	SizeX = Width;
+	SizeY = Height;
+	UpdateTextureRaw(Buffer);
+}
+
+void FSlateOpenGLTexture::UpdateTextureRaw(const void* Buffer)
 {
 	// Ensure texturing is enabled before setting texture properties
 #if USE_DEPRECATED_OPENGL_FUNCTIONALITY
-    glEnable(GL_TEXTURE_2D);
+	glEnable(GL_TEXTURE_2D);
 #endif // USE_DEPRECATED_OPENGL_FUNCTIONALITY
 	glBindTexture(GL_TEXTURE_2D, ShaderResource);
 	glTexParameteri(GL_TEXTURE_2D, GL_TEXTURE_MIN_FILTER, GL_LINEAR);
 	glTexParameteri(GL_TEXTURE_2D, GL_TEXTURE_MAG_FILTER, GL_LINEAR);
+	glTexParameteri(GL_TEXTURE_2D, GL_TEXTURE_WRAP_S, GL_CLAMP_TO_EDGE);
 #if USE_DEPRECATED_OPENGL_FUNCTIONALITY
-    glTexEnvi(GL_TEXTURE_ENV, GL_TEXTURE_ENV_MODE, GL_REPLACE);
+	glTexEnvi(GL_TEXTURE_ENV, GL_TEXTURE_ENV_MODE, GL_REPLACE);
 #endif // USE_DEPRECATED_OPENGL_FUNCTIONALITY
-
+	
 	// Upload the texture data
 #if !PLATFORM_USES_ES2
-	glTexImage2D(GL_TEXTURE_2D, 0, GL_SRGB8_ALPHA8, SizeX, SizeY, 0, GL_RGBA, GL_UNSIGNED_INT_8_8_8_8_REV, Bytes.GetData());
+	glTexImage2D(GL_TEXTURE_2D, 0, GL_SRGB8_ALPHA8, SizeX, SizeY, 0, GL_RGBA, GL_UNSIGNED_INT_8_8_8_8_REV, Buffer);
 #else
-	glTexImage2D(GL_TEXTURE_2D, 0, GL_SRGB8_ALPHA8_EXT, SizeX, SizeY, 0, GL_RGBA, GL_UNSIGNED_INT_8_8_8_8_REV, Bytes.GetData());
+	glTexImage2D(GL_TEXTURE_2D, 0, GL_SRGB8_ALPHA8_EXT, SizeX, SizeY, 0, GL_RGBA, GL_UNSIGNED_INT_8_8_8_8_REV, Buffer);
 #endif
 	CHECK_GL_ERRORS;
 }
 
+
 FSlateFontTextureOpenGL::FSlateFontTextureOpenGL( uint32 Width, uint32 Height )
 	: FSlateFontAtlas( Width, Height ) 
 	, FontTexture(nullptr)
@@ -90,15 +104,15 @@ void FSlateFontTextureOpenGL::CreateFontTexture()
 	glTexParameteri(GL_TEXTURE_2D, GL_TEXTURE_MAG_FILTER, GL_NEAREST);
 
 	// Font textures use an alpha texture only
-    GLint Format =
+	GLint Format =
 #if USE_DEPRECATED_OPENGL_FUNCTIONALITY
-        GL_ALPHA
+		GL_ALPHA
 #else
 		GL_RED
 #endif // USE_DEPRECATED_OPENGL_FUNCTIONALITY
-    ;
+	;
 
-    // Upload the data to the texture
+	// Upload the data to the texture
 	glTexImage2D( GL_TEXTURE_2D, 0, Format, AtlasWidth, AtlasHeight, 0, Format, GL_UNSIGNED_BYTE, NULL );
 
 	// Create a new slate texture for use in rendering
@@ -117,13 +131,13 @@ void FSlateFontTextureOpenGL::ConditionalUpdateTexture()
 
 		// Completely the texture data each time characters are added
 		glBindTexture(GL_TEXTURE_2D, FontTexture->GetTypedResource() );
-        GLint Format =
+		GLint Format =
 #if USE_DEPRECATED_OPENGL_FUNCTIONALITY
-            GL_ALPHA
+			GL_ALPHA
 #else
-            GL_RED
+			GL_RED
 #endif // USE_DEPRECATED_OPENGL_FUNCTIONALITY
-        ;
+		;
 
 #if PLATFORM_MAC // Make this texture use a DMA'd client storage backing store on OS X, where these extensions always exist
 				 // This avoids a problem on Intel & Nvidia cards that makes characters disappear as well as making the texture updates
diff --git a/Engine/Source/Developer/StandaloneRenderer/Private/OpenGL/SlateOpenGLTextures.h b/Engine/Source/Developer/StandaloneRenderer/Private/OpenGL/SlateOpenGLTextures.h
--- a/Engine/Source/Developer/StandaloneRenderer/Private/OpenGL/SlateOpenGLTextures.h
+++ b/Engine/Source/Developer/StandaloneRenderer/Private/OpenGL/SlateOpenGLTextures.h
@@ -33,7 +33,11 @@ public:
 	virtual void ResizeTexture( uint32 Width, uint32 Height ) override;
 	virtual void UpdateTexture(const TArray<uint8>& Bytes) override;
 	virtual void UpdateTextureThreadSafe(const TArray<uint8>& Bytes) override { UpdateTexture(Bytes); }
+	virtual void UpdateTextureThreadSafeRaw(uint32 Width, uint32 Height, const void* Buffer) override;
 private:
+	// Helper method used by the different UpdateTexture* methods
+	void UpdateTextureRaw(const void* Buffer);
+
 	static GLuint NullTexture;
 
 	uint32 SizeX;
diff --git a/Engine/Source/Developer/StandaloneRenderer/Private/Windows/D3D/SlateD3DTextures.cpp b/Engine/Source/Developer/StandaloneRenderer/Private/Windows/D3D/SlateD3DTextures.cpp
--- a/Engine/Source/Developer/StandaloneRenderer/Private/Windows/D3D/SlateD3DTextures.cpp
+++ b/Engine/Source/Developer/StandaloneRenderer/Private/Windows/D3D/SlateD3DTextures.cpp
@@ -43,18 +43,33 @@ void FSlateD3DTexture::ResizeTexture(uint32 Width, uint32 Height)
 	Init(DXGI_FORMAT_B8G8R8A8_UNORM_SRGB, NULL, D3D11_USAGE_DYNAMIC, D3D11_CPU_ACCESS_WRITE);
 }
 
-void FSlateD3DTexture::UpdateTexture(const TArray<uint8>& Bytes)
+void FSlateD3DTexture::UpdateTextureRaw(const void* Buffer)
 {
 	// TODO: Improve the memory copying here, have tried using UpdateSubresource but it doesn't seem to work
 	D3D11_MAPPED_SUBRESOURCE Resource;
 	GD3DDeviceContext->Map(D3DTexture, 0, D3D11_MAP_WRITE_DISCARD, 0, &Resource);
 	for (uint32 Row = 0; Row < SizeY; ++Row)
 	{
-		FMemory::Memcpy((uint8*)Resource.pData + Row * Resource.RowPitch, Bytes.GetData() + Row * SizeX*4, SizeX*4);
+		FMemory::Memcpy((uint8*)Resource.pData + Row * Resource.RowPitch, (uint8*)Buffer + Row * SizeX*4, SizeX*4);
 	}
 	GD3DDeviceContext->Unmap(D3DTexture, 0);
 }
 
+void FSlateD3DTexture::UpdateTexture(const TArray<uint8>& Bytes)
+{
+	UpdateTextureRaw(Bytes.GetData());
+}
+
+void FSlateD3DTexture::UpdateTextureThreadSafeRaw(uint32 Width, uint32 Height, const void* Buffer)
+{
+	if (Width != SizeX || Height != SizeY)
+	{
+		ResizeTexture(Width, Height);
+	}
+	UpdateTextureRaw(Buffer);
+}
+
+
 FSlateTextureAtlasD3D::FSlateTextureAtlasD3D( uint32 Width, uint32 Height, uint32 StrideBytes, ESlateTextureAtlasPaddingStyle PaddingStyle )
 	: FSlateTextureAtlas( Width, Height, StrideBytes, PaddingStyle )
 	, AtlasTexture( new FSlateD3DTexture( Width, Height ) )
diff --git a/Engine/Source/Developer/StandaloneRenderer/Private/Windows/D3D/SlateD3DTextures.h b/Engine/Source/Developer/StandaloneRenderer/Private/Windows/D3D/SlateD3DTextures.h
--- a/Engine/Source/Developer/StandaloneRenderer/Private/Windows/D3D/SlateD3DTextures.h
+++ b/Engine/Source/Developer/StandaloneRenderer/Private/Windows/D3D/SlateD3DTextures.h
@@ -34,7 +34,10 @@ public:
 	virtual void ResizeTexture( uint32 Width, uint32 Height ) override;
 	virtual void UpdateTexture(const TArray<uint8>& Bytes) override;
 	virtual void UpdateTextureThreadSafe(const TArray<uint8>& Bytes) override { UpdateTexture(Bytes); }
+	virtual void UpdateTextureThreadSafeRaw(uint32 Width, uint32 Height, const void* Buffer) override;
 private:
+	// Helper method used by the different UpdateTexture* methods
+	void UpdateTextureRaw(const void* Buffer);
 	/** Actual texture.  In D3D the SRV is used by the shader so our parent class holds that */
 	TRefCountPtr<ID3D11Texture2D> D3DTexture;
 	uint32 SizeX;
diff --git a/Engine/Source/Runtime/Engine/Private/Slate/SlateTextures.cpp b/Engine/Source/Runtime/Engine/Private/Slate/SlateTextures.cpp
--- a/Engine/Source/Runtime/Engine/Private/Slate/SlateTextures.cpp
+++ b/Engine/Source/Runtime/Engine/Private/Slate/SlateTextures.cpp
@@ -181,6 +181,33 @@ void FSlateTexture2DRHIRef::UpdateTextureThreadSafe(const TArray<uint8>& Bytes)
 	}
 }
 
+void FSlateTexture2DRHIRef::UpdateTextureThreadSafeRaw(uint32 InWidth, uint32 InHeight, const void* Buffer)
+{
+	if (IsInGameThread())
+	{
+		// No cheap way to avoid having to copy the Buffer, as we cannot guarantee it will not be touched before the rendering thread is done with it.
+		const uint32 BufferSize = InWidth * InHeight * 4;
+		TArray<uint8>* BufferCopy = new TArray<uint8>();
+		BufferCopy->AddUninitialized(BufferSize);
+		FMemory::Memcpy(BufferCopy->GetData(), Buffer, BufferSize);
+		FIntPoint Dimensions(InWidth, InHeight);
+
+		// Update the texture RHI
+		ENQUEUE_UNIQUE_RENDER_COMMAND_THREEPARAMETER(
+			FSlateTexture2DRHIRef_UpdateTextureRaw,
+			FSlateTexture2DRHIRef*, ThisTexture, this,
+			FIntPoint, InDimensions, Dimensions,
+			TArray<uint8>*, BufferCopy, BufferCopy,
+			{
+				if (ThisTexture->GetWidth() != InDimensions.X || ThisTexture->GetHeight() != InDimensions.Y)
+				{
+					ThisTexture->Resize(InDimensions.X, InDimensions.Y);
+				}
+				ThisTexture->UpdateTexture(*BufferCopy);
+				delete BufferCopy;
+			});
+	}
+}
 
 void FSlateRenderTargetRHI::SetRHIRef( FTexture2DRHIRef InRenderTargetTexture, uint32 InWidth, uint32 InHeight )
 {
@@ -224,8 +251,8 @@ void FSlateTextureRenderTarget2DResource::ClampSize(int32 MaxSizeX,int32 MaxSize
 	check(IsInRenderingThread());
 
 	// upsize to go back to original or downsize to clamp to max
- 	int32 NewSizeX = FMath::Min<int32>(TargetSizeX,MaxSizeX);
- 	int32 NewSizeY = FMath::Min<int32>(TargetSizeY,MaxSizeY);
+	int32 NewSizeX = FMath::Min<int32>(TargetSizeX,MaxSizeX);
+	int32 NewSizeY = FMath::Min<int32>(TargetSizeY,MaxSizeY);
 	if (NewSizeX != TargetSizeX || NewSizeY != TargetSizeY)
 	{
 		TargetSizeX = NewSizeX;
diff --git a/Engine/Source/Runtime/Engine/Public/Slate/SlateTextures.h b/Engine/Source/Runtime/Engine/Public/Slate/SlateTextures.h
--- a/Engine/Source/Runtime/Engine/Public/Slate/SlateTextures.h
+++ b/Engine/Source/Runtime/Engine/Public/Slate/SlateTextures.h
@@ -74,6 +74,7 @@ public:
 	virtual void ResizeTexture( uint32 Width, uint32 Height ) override;
 	virtual void UpdateTexture(const TArray<uint8>& Bytes) override;
 	virtual void UpdateTextureThreadSafe(const TArray<uint8>& Bytes) override;
+	virtual void UpdateTextureThreadSafeRaw(uint32 Width, uint32 Height, const void* Buffer) override;
 protected:
 	/** Width of this texture */
 	uint32 Width;
diff --git a/Engine/Source/Runtime/SlateCore/Public/Textures/SlateUpdatableTexture.h b/Engine/Source/Runtime/SlateCore/Public/Textures/SlateUpdatableTexture.h
--- a/Engine/Source/Runtime/SlateCore/Public/Textures/SlateUpdatableTexture.h
+++ b/Engine/Source/Runtime/SlateCore/Public/Textures/SlateUpdatableTexture.h
@@ -50,4 +50,17 @@ public:
 	 * @param Bytes Array of texture data
 	 */
 	virtual void UpdateTextureThreadSafe(const TArray<uint8>& Bytes) = 0;
+	
+	/**
+	 * Update the texture from a raw byte buffer.
+	 * Should only be used when integrating with third party APIs that provide a raw pointer to texture data.
+	 * This method does a copy of the buffer for thread safety. 
+	 * Will resize the texture if the passed in width and height is different from the current size.
+	 * The passed in size must correspond to the size of the buffer.
+	 *
+	 * @param Width New texture width
+	 * @param Height New texture height
+	 * @param Buffer a void pointer to a byte buffer.
+	 */
+	virtual void UpdateTextureThreadSafeRaw(uint32 Width, uint32 Height, const void* Buffer) = 0;
 };
