Rewrite Mac OpenGL handling for CoreAnimation so that we don't hold a pointer to the FSlateOpenGLViewport as that it is an entry in a TMap & can be reallocated at any time.

diff --git a/Engine/Source/Developer/StandaloneRenderer/Private/Mac/OpenGL/SlateOpenGLContext.cpp b/Engine/Source/Developer/StandaloneRenderer/Private/Mac/OpenGL/SlateOpenGLContext.cpp
--- a/Engine/Source/Developer/StandaloneRenderer/Private/Mac/OpenGL/SlateOpenGLContext.cpp
+++ b/Engine/Source/Developer/StandaloneRenderer/Private/Mac/OpenGL/SlateOpenGLContext.cpp
@@ -2,17 +2,15 @@
 
 #include "StandaloneRendererPrivate.h"
 #include "OpenGL/SlateOpenGLRenderer.h"
+
+#include "SlateOpenGLMac.h"
+
 #include "MacWindow.h"
 #include "MacTextInputMethodSystem.h"
 #include "CocoaTextView.h"
 #include <OpenGL/gl.h>
 #include <OpenGL/glext.h>
 
-@interface FSlateOpenGLLayer : NSOpenGLLayer
-@property (assign) NSOpenGLContext* Context;
-@property (assign) NSOpenGLPixelFormat* PixelFormat;
-@end
-
 @implementation FSlateOpenGLLayer
 
 - (NSOpenGLPixelFormat *)openGLPixelFormatForDisplayMask:(uint32)Mask
@@ -38,12 +36,6 @@
 
 @end
 
-@interface FSlateCocoaView : FCocoaTextView
-@property (assign) NSOpenGLContext* Context;
-@property (assign) NSOpenGLPixelFormat* PixelFormat;
-@property (assign) FSlateOpenGLViewport* SlateViewport;
-@end
-
 @implementation FSlateCocoaView
 
 - (CALayer*)makeBackingLayer
@@ -51,21 +43,41 @@
 	return [[FSlateOpenGLLayer alloc] initWithContext:self.Context andPixelFormat:self.PixelFormat];
 }
 
-- (id)initWithFrame:(NSRect)frameRect context:(NSOpenGLContext*)context pixelFormat:(NSOpenGLPixelFormat*)pixelFormat andSlateViewport:(FSlateOpenGLViewport*)slateViewport
+- (id)initWithFrame:(NSRect)frameRect context:(NSOpenGLContext*)context pixelFormat:(NSOpenGLPixelFormat*)pixelFormat
 {
 	self = [super initWithFrame:frameRect];
 	if (self)
 	{
 		self.Context = context;
 		self.PixelFormat = pixelFormat;
-		self.SlateViewport = slateViewport;
+		Framebuffer = 0;
+		Renderbuffer = 0;
 	}
 	return self;
 }
 
+-(void)dealloc
+{
+	[super dealloc];
+}
+
 - (void)drawRect:(NSRect)dirtyRect
 {
-	self.SlateViewport->Draw();
+	if (Framebuffer && [(FCocoaWindow*)[self window] isRenderInitialized] && ViewportRect.IsValid())
+	{
+		int32 CurrentReadFramebuffer = 0;
+		glGetIntegerv(GL_READ_FRAMEBUFFER_BINDING, &CurrentReadFramebuffer);
+		glBindFramebuffer(GL_READ_FRAMEBUFFER, Framebuffer);
+		glReadBuffer(GL_COLOR_ATTACHMENT0);
+		glBlitFramebuffer(0, 0, ViewportRect.Right, ViewportRect.Bottom, 0, 0, self.frame.size.width, self.frame.size.height, GL_COLOR_BUFFER_BIT, GL_NEAREST);
+		CHECK_GL_ERRORS;
+		glBindFramebuffer(GL_READ_FRAMEBUFFER, CurrentReadFramebuffer);
+	}
+	else
+	{
+		glClearColor(0.0f, 0.0f, 0.0f, 0.0f);
+		glClear(GL_COLOR_BUFFER_BIT);
+	}
 }
 
 - (BOOL)isOpaque
@@ -93,7 +105,6 @@ FSlateOpenGLContext::FSlateOpenGLContext()
 :	View(NULL)
 ,	PixelFormat(NULL)
 ,	Context(NULL)
-,	ViewContext(NULL)
 ,	bNeedsUpdate(false)
 {
 }
@@ -122,9 +133,8 @@ void FSlateOpenGLContext::Initialize(void* InWindow, const FSlateOpenGLContext*
 	NSWindow* Window = (NSWindow*)InWindow;
 	if (Window)
 	{
-		ViewContext = [[NSOpenGLContext alloc] initWithFormat:PixelFormat shareContext:Context];
 		const NSRect ViewRect = NSMakeRect(0, 0, Window.frame.size.width, Window.frame.size.height);
-		View = [[FSlateCocoaView alloc] initWithFrame:ViewRect context:ViewContext pixelFormat:PixelFormat andSlateViewport:Viewport];
+		View = [[FSlateCocoaView alloc] initWithFrame:ViewRect context:Context pixelFormat:PixelFormat];
 		[View setAutoresizingMask:NSViewWidthSizable | NSViewHeightSizable];
 
 		if (FPlatformMisc::IsRunningOnMavericks() && ([Window styleMask] & NSTexturedBackgroundWindowMask))
@@ -163,6 +173,22 @@ void FSlateOpenGLContext::Destroy()
 		{
 			[Window setContentView:NULL];
 		}
+		
+		NSOpenGLContext* Current = [NSOpenGLContext currentContext];
+		[Context makeCurrentContext];
+		FSlateCocoaView* SlateView = ((FSlateCocoaView*)View);
+		if (SlateView->Framebuffer)
+		{
+			glDeleteFramebuffers(1, &SlateView->Framebuffer);
+			SlateView->Framebuffer = 0;
+		}
+		if (SlateView->Renderbuffer)
+		{
+			glDeleteRenderbuffers(1, &SlateView->Renderbuffer);
+			SlateView->Renderbuffer = 0;
+		}
+		[Current makeCurrentContext];
+		
 		[View release];
 		View = NULL;
 
@@ -172,11 +198,8 @@ void FSlateOpenGLContext::Destroy()
 		[PixelFormat release];
 		[Context clearDrawable];
 		[Context release];
-		[ViewContext clearDrawable];
-		[ViewContext release];
 		PixelFormat = NULL;
 		Context = NULL;
-		ViewContext = NULL;
 	}
 }
 
diff --git a/Engine/Source/Developer/StandaloneRenderer/Private/Mac/OpenGL/SlateOpenGLMac.h b/Engine/Source/Developer/StandaloneRenderer/Private/Mac/OpenGL/SlateOpenGLMac.h
new file mode 100644
--- /dev/null
+++ b/Engine/Source/Developer/StandaloneRenderer/Private/Mac/OpenGL/SlateOpenGLMac.h
@@ -0,0 +1,27 @@
+// Copyright 1998-2015 Epic Games, Inc. All Rights Reserved.
+
+#pragma once 
+
+#include "StandaloneRendererPrivate.h"
+#include "OpenGL/SlateOpenGLRenderer.h"
+#include "MacWindow.h"
+#include "MacTextInputMethodSystem.h"
+#include "CocoaTextView.h"
+#include <OpenGL/gl.h>
+#include <OpenGL/glext.h>
+
+@interface FSlateOpenGLLayer : NSOpenGLLayer
+@property (assign) NSOpenGLContext* Context;
+@property (assign) NSOpenGLPixelFormat* PixelFormat;
+@end
+
+@interface FSlateCocoaView : FCocoaTextView
+{
+@public
+	GLuint Framebuffer;
+	GLuint Renderbuffer;
+	FSlateRect ViewportRect;
+}
+@property (assign) NSOpenGLContext* Context;
+@property (assign) NSOpenGLPixelFormat* PixelFormat;
+@end
diff --git a/Engine/Source/Developer/StandaloneRenderer/Private/Mac/OpenGL/SlateOpenGLViewport.cpp b/Engine/Source/Developer/StandaloneRenderer/Private/Mac/OpenGL/SlateOpenGLViewport.cpp
--- a/Engine/Source/Developer/StandaloneRenderer/Private/Mac/OpenGL/SlateOpenGLViewport.cpp
+++ b/Engine/Source/Developer/StandaloneRenderer/Private/Mac/OpenGL/SlateOpenGLViewport.cpp
@@ -2,18 +2,18 @@
 
 #include "StandaloneRendererPrivate.h"
 #include "OpenGL/SlateOpenGLRenderer.h"
+
+#include "SlateOpenGLMac.h"
+
 #include "MacWindow.h"
 
 FSlateOpenGLViewport::FSlateOpenGLViewport()
-:	Framebuffer(0)
-,	Renderbuffer(0)
 {
 }
 
 void FSlateOpenGLViewport::Initialize(TSharedRef<SWindow> InWindow, const FSlateOpenGLContext& SharedContext)
 {
 	TSharedRef<FGenericWindow> NativeWindow = InWindow->GetNativeWindow().ToSharedRef();
-	RenderingContext.Viewport = this;
 	RenderingContext.Initialize(NativeWindow->GetOSWindowHandle(), &SharedContext);
 
 	const int32 Width = FMath::TruncToInt(InWindow->GetSizeInScreen().X);
@@ -27,24 +27,13 @@ void FSlateOpenGLViewport::Initialize(TSharedRef<SWindow> InWindow, const FSlate
 
 void FSlateOpenGLViewport::Destroy()
 {
-	if (Framebuffer)
-	{
-		glDeleteFramebuffers(1, &Framebuffer);
-		Framebuffer = 0;
-	}
-	if (Renderbuffer)
-	{
-		glDeleteRenderbuffers(1, &Renderbuffer);
-		Renderbuffer = 0;
-	}
 	RenderingContext.Destroy();
 }
 
 void FSlateOpenGLViewport::MakeCurrent()
 {
-	ContextGuard.Lock();
 	RenderingContext.MakeCurrent();
-	glBindFramebuffer(GL_FRAMEBUFFER, Framebuffer);
+	glBindFramebuffer(GL_FRAMEBUFFER, ((FSlateCocoaView*)RenderingContext.View)->Framebuffer);
 }
 
 void FSlateOpenGLViewport::SwapBuffers()
@@ -52,68 +41,51 @@ void FSlateOpenGLViewport::SwapBuffers()
 	glFlushRenderAPPLE();
 	[(FCocoaWindow*)[RenderingContext.View window] startRendering];
 	[RenderingContext.View setNeedsDisplay:YES];
-	ContextGuard.Unlock();
+	glBindFramebuffer(GL_FRAMEBUFFER, ((FSlateCocoaView*)RenderingContext.View)->Framebuffer);
 }
 
 void FSlateOpenGLViewport::Resize(int32 Width, int32 Height, bool bInFullscreen)
 {
-	FScopeLock Lock(&ContextGuard);
-
 	ViewportRect.Right = Width;
 	ViewportRect.Bottom = Height;
 
 	// Need to create a new projection matrix each time the window is resized
 	ProjectionMatrix = CreateProjectionMatrix(Width, Height);
 
-	if (RenderingContext.Context && Width > 0 && Height > 0)
+	if (RenderingContext.Context && RenderingContext.View && Width > 0 && Height > 0)
 	{
-		[RenderingContext.Context update];
-
+		((FSlateCocoaView*)RenderingContext.View)->ViewportRect = ViewportRect;
+		
+		GLuint& Framebuffer = ((FSlateCocoaView*)RenderingContext.View)->Framebuffer;
+		GLuint& Renderbuffer = ((FSlateCocoaView*)RenderingContext.View)->Renderbuffer;
+		
 		if (Framebuffer == 0)
 		{
 			glGenFramebuffers(1, &Framebuffer);
 			check(Framebuffer);
 		}
-
+		
 		if (Renderbuffer == 0)
 		{
 			glGenRenderbuffers(1, &Renderbuffer);
 			check(Renderbuffer);
 		}
-
+		
 		int32 CurrentDrawFramebuffer = 0;
 		int32 CurrentReadFramebuffer = 0;
 		glGetIntegerv(GL_DRAW_FRAMEBUFFER_BINDING, &CurrentDrawFramebuffer);
 		glGetIntegerv(GL_READ_FRAMEBUFFER_BINDING, &CurrentReadFramebuffer);
-
+		
 		glBindRenderbuffer(GL_RENDERBUFFER, Renderbuffer);
 		glRenderbufferStorage(GL_RENDERBUFFER, GL_RGB, Width, Height);
 		glBindRenderbuffer(GL_RENDERBUFFER, 0);
-
+		
 		glBindFramebuffer(GL_FRAMEBUFFER, Framebuffer);
 		glFramebufferRenderbuffer(GL_FRAMEBUFFER, GL_COLOR_ATTACHMENT0, GL_RENDERBUFFER, Renderbuffer);
-
+		
 		glBindFramebuffer(GL_DRAW_FRAMEBUFFER, CurrentDrawFramebuffer);
 		glBindFramebuffer(GL_READ_FRAMEBUFFER, CurrentReadFramebuffer);
-	}
-}
 
-void FSlateOpenGLViewport::Draw()
-{
-	FScopeLock Lock(&ContextGuard);
-	if (RenderingContext.View && Framebuffer && [(FCocoaWindow*)[RenderingContext.View window] isRenderInitialized])
-	{
-		int32 CurrentReadFramebuffer = 0;
-		glGetIntegerv(GL_READ_FRAMEBUFFER_BINDING, &CurrentReadFramebuffer);
-		glBindFramebuffer(GL_READ_FRAMEBUFFER, Framebuffer);
-		glReadBuffer(GL_COLOR_ATTACHMENT0);
-		glBlitFramebuffer(0, 0, ViewportRect.Right, ViewportRect.Bottom, 0, 0, RenderingContext.View.frame.size.width, RenderingContext.View.frame.size.height, GL_COLOR_BUFFER_BIT, GL_NEAREST);
-		CHECK_GL_ERRORS;
-		glBindFramebuffer(GL_READ_FRAMEBUFFER, CurrentReadFramebuffer);
-	}
-	else
-	{
-		glClearColor(0.0f, 0.0f, 0.0f, 0.0f);
-		glClear(GL_COLOR_BUFFER_BIT);
+		[RenderingContext.Context update];
 	}
 }
diff --git a/Engine/Source/Developer/StandaloneRenderer/Private/OpenGL/SlateOpenGLRenderer.h b/Engine/Source/Developer/StandaloneRenderer/Private/OpenGL/SlateOpenGLRenderer.h
--- a/Engine/Source/Developer/StandaloneRenderer/Private/OpenGL/SlateOpenGLRenderer.h
+++ b/Engine/Source/Developer/StandaloneRenderer/Private/OpenGL/SlateOpenGLRenderer.h
@@ -35,8 +35,6 @@ struct FSlateOpenGLContext
 	NSView* View;
 	NSOpenGLPixelFormat* PixelFormat;
 	NSOpenGLContext* Context;
-	NSOpenGLContext* ViewContext; // context for CALayer
-	struct FSlateOpenGLViewport* Viewport;
 	bool bNeedsUpdate;
 #elif PLATFORM_IOS
 	UIWindow* WindowHandle;
@@ -101,10 +99,6 @@ struct FSlateOpenGLViewport
 	/** Updates the OpenGL context and projection matrix on resize */
 	void Resize( int32 Width, int32 Height, bool bInFullscreen );
 
-#if PLATFORM_MAC
-	void Draw();
-#endif
-
 private:
 
 	/**
@@ -129,12 +123,6 @@ private:
 						 FPlane((Left+Right)/(Left-Right),	(Top+Bottom)/(Bottom-Top),	ZNear/(ZNear-ZFar), 1 ) );
 
 	}
-
-#if PLATFORM_MAC
-	GLuint Framebuffer;
-	GLuint Renderbuffer;
-	FCriticalSection ContextGuard;
-#endif
 };
 
 /**
