Users with unsupported video cards will now get an OS error dialog explaining the issue.
OPP-3018

Merging using UE4-To-UE4-LauncherDev

diff --git a/Engine/Source/Developer/StandaloneRenderer/Private/Windows/D3D/SlateD3DRenderer.cpp b/Engine/Source/Developer/StandaloneRenderer/Private/Windows/D3D/SlateD3DRenderer.cpp
--- a/Engine/Source/Developer/StandaloneRenderer/Private/Windows/D3D/SlateD3DRenderer.cpp
+++ b/Engine/Source/Developer/StandaloneRenderer/Private/Windows/D3D/SlateD3DRenderer.cpp
@@ -115,6 +115,12 @@ void FSlateD3DRenderer::CreateDevice()
 		D3D_FEATURE_LEVEL CreatedFeatureLevel;
 		HRESULT Hr = D3D11CreateDevice( NULL, DriverType, NULL, DeviceCreationFlags, FeatureLevels, sizeof(FeatureLevels)/sizeof(D3D_FEATURE_LEVEL), D3D11_SDK_VERSION, GD3DDevice.GetInitReference(), &CreatedFeatureLevel, GD3DDeviceContext.GetInitReference() );
 		UE_LOG(LogStandaloneRenderer, Log, TEXT("D3D11CreateDevice Result: %X"), Hr);
+
+		if (Hr == DXGI_ERROR_UNSUPPORTED)
+		{
+			FPlatformMisc::MessageBoxExt(EAppMsgType::Ok, *NSLOCTEXT("", "", "There is a problem with your graphics card. Please ensure your card meets the minimum system requirements and that you have the latest drivers installed.").ToString(), *NSLOCTEXT("", "UnsupportedVideoCardErrorTitle", "Unsupported Video Card").ToString());
+		}
+		
 		checkf(SUCCEEDED(Hr), TEXT("D3D11 Error Result %X"), Hr);
 	}
 }
