SlateViewer: fix compilation on Linux.

- Still broken in runtime.

diff --git a/Engine/Source/Developer/StandaloneRenderer/Private/Linux/OpenGL/SlateOpenGLContext.cpp b/Engine/Source/Developer/StandaloneRenderer/Private/Linux/OpenGL/SlateOpenGLContext.cpp
--- a/Engine/Source/Developer/StandaloneRenderer/Private/Linux/OpenGL/SlateOpenGLContext.cpp
+++ b/Engine/Source/Developer/StandaloneRenderer/Private/Linux/OpenGL/SlateOpenGLContext.cpp
@@ -3,6 +3,150 @@
 #include "StandaloneRendererPrivate.h"
 #include "OpenGL/SlateOpenGLRenderer.h"
 
+#if defined(GL_ARB_debug_output) || defined(GL_KHR_debug)
+
+/**
+ * Map GL_DEBUG_SOURCE_*_ARB to a human-readable string.
+ */
+static const TCHAR* GetOpenGLDebugSourceStringARB(GLenum Source)
+{
+	static const TCHAR* SourceStrings[] =
+	{
+		TEXT("API"),
+		TEXT("System"),
+		TEXT("ShaderCompiler"),
+		TEXT("ThirdParty"),
+		TEXT("Application"),
+		TEXT("Other")
+	};
+
+	if (Source >= GL_DEBUG_SOURCE_API_ARB && Source <= GL_DEBUG_SOURCE_OTHER_ARB)
+	{
+		return SourceStrings[Source - GL_DEBUG_SOURCE_API_ARB];
+	}
+	return TEXT("Unknown");
+}
+
+/**
+ * Map GL_DEBUG_TYPE_*_ARB to a human-readable string.
+ */
+static const TCHAR* GetOpenGLDebugTypeStringARB(GLenum Type)
+{
+	static const TCHAR* TypeStrings[] =
+	{
+		TEXT("Error"),
+		TEXT("Deprecated"),
+		TEXT("UndefinedBehavior"),
+		TEXT("Portability"),
+		TEXT("Performance"),
+		TEXT("Other")
+	};
+
+	if (Type >= GL_DEBUG_TYPE_ERROR_ARB && Type <= GL_DEBUG_TYPE_OTHER_ARB)
+	{
+		return TypeStrings[Type - GL_DEBUG_TYPE_ERROR_ARB];
+	}
+#ifdef GL_KHR_debug
+	{
+		static const TCHAR* TypeStrings[] =
+		{
+			TEXT("Marker"),
+			TEXT("PushGroup"),
+			TEXT("PopGroup"),
+		};
+
+		if (Type >= GL_DEBUG_TYPE_MARKER && Type <= GL_DEBUG_TYPE_POP_GROUP)
+		{
+			return TypeStrings[Type - GL_DEBUG_TYPE_MARKER];
+		}
+	}
+#endif
+	return TEXT("Unknown");
+}
+
+/**
+ * Map GL_DEBUG_SEVERITY_*_ARB to a human-readable string.
+ */
+static const TCHAR* GetOpenGLDebugSeverityStringARB(GLenum Severity)
+{
+	static const TCHAR* SeverityStrings[] =
+	{
+		TEXT("High"),
+		TEXT("Medium"),
+		TEXT("Low")
+	};
+
+	if (Severity >= GL_DEBUG_SEVERITY_HIGH_ARB && Severity <= GL_DEBUG_SEVERITY_LOW_ARB)
+	{
+		return SeverityStrings[Severity - GL_DEBUG_SEVERITY_HIGH_ARB];
+	}
+#ifdef GL_KHR_debug
+	if(Severity == GL_DEBUG_SEVERITY_NOTIFICATION)
+		return TEXT("Notification");
+#endif
+	return TEXT("Unknown");
+}
+
+/**
+ * OpenGL debug message callback. Conforms to GLDEBUGPROCARB.
+ */
+static void APIENTRY OpenGLDebugMessageCallbackARB(
+	GLenum Source,
+	GLenum Type,
+	GLuint Id,
+	GLenum Severity,
+	GLsizei Length,
+	const GLchar* Message,
+	GLvoid* UserParam)
+{
+#if !NO_LOGGING
+	const TCHAR* SourceStr = GetOpenGLDebugSourceStringARB(Source);
+	const TCHAR* TypeStr = GetOpenGLDebugTypeStringARB(Type);
+	const TCHAR* SeverityStr = GetOpenGLDebugSeverityStringARB(Severity);
+
+	ELogVerbosity::Type Verbosity = ELogVerbosity::Warning;
+	if (Type == GL_DEBUG_TYPE_ERROR_ARB && Severity == GL_DEBUG_SEVERITY_HIGH_ARB)
+	{
+		Verbosity = ELogVerbosity::Fatal;
+	}
+
+	if ((Verbosity & ELogVerbosity::VerbosityMask) <= FLogCategoryLogStandaloneRenderer::CompileTimeVerbosity)
+	{
+		if (!LogStandaloneRenderer.IsSuppressed(Verbosity))
+		{
+			FMsg::Logf(__FILE__, __LINE__, LogStandaloneRenderer.GetCategoryName(), Verbosity,
+				TEXT("[%s][%s][%s][%u] %s"),
+				SourceStr,
+				TypeStr,
+				SeverityStr,
+				Id,
+				ANSI_TO_TCHAR(Message)
+				);
+		}
+
+		// this is a debugging code to catch VIDEO->HOST copying
+		if (Id == 131186)
+		{
+			int A = 5;
+		}
+	}
+#endif
+}
+
+#endif // defined(GL_ARB_debug_output) || defined(GL_KHR_debug)
+
+/**
+ * Enable/Disable debug context from the commandline
+ */
+static bool PlatformOpenGLDebugCtx()
+{
+#if UE_BUILD_DEBUG
+	return ! FParse::Param(FCommandLine::Get(),TEXT("openglNoDebug"));
+#else
+	return FParse::Param(FCommandLine::Get(),TEXT("openglDebug"));;
+#endif
+}
+
 static SDL_Window* CreateDummyGLWindow()
 {
 	FPlatformMisc::PlatformInitMultimedia(); //	will not initialize more than once
@@ -62,7 +206,11 @@ void FSlateOpenGLContext::Initialize( void* InWindow, const FSlateOpenGLContext*
 	SDL_GL_SetAttribute( SDL_GL_CONTEXT_MAJOR_VERSION, OpenGLMajorVersionToUse );
 	SDL_GL_SetAttribute( SDL_GL_CONTEXT_MINOR_VERSION, OpenGLMinorVersionToUse );
 	SDL_GL_SetAttribute( SDL_GL_CONTEXT_PROFILE_MASK, SDL_GL_CONTEXT_PROFILE_CORE );
-	SDL_GL_SetAttribute( SDL_GL_CONTEXT_FLAGS, SDL_GL_CONTEXT_DEBUG_FLAG );
+
+	if (PlatformOpenGLDebugCtx())
+	{
+		SDL_GL_SetAttribute(SDL_GL_CONTEXT_FLAGS, SDL_GL_CONTEXT_DEBUG_FLAG);
+	}
 
 	if( SharedContext )
 	{
@@ -97,6 +245,30 @@ void FSlateOpenGLContext::Initialize( void* InWindow, const FSlateOpenGLContext*
 
 	LoadOpenGLExtensions();
 
+	if (PlatformOpenGLDebugCtx())
+	{
+		if (glDebugMessageCallbackARB)
+		{
+		#if GL_ARB_debug_output || GL_KHR_debug
+			// Synchronous output can slow things down, but we'll get better callstack if breaking in or crashing in the callback. This is debug only after all.
+			glEnable(GL_DEBUG_OUTPUT_SYNCHRONOUS);
+			glDebugMessageCallbackARB(GLDEBUGPROCARB(OpenGLDebugMessageCallbackARB), /*UserParam=*/ nullptr);
+			checkf(glGetError() == GL_NO_ERROR, TEXT("Could not register glDebugMessageCallbackARB()"));
+
+			if(glDebugMessageControlARB)
+			{
+				glDebugMessageControlARB(GL_DEBUG_SOURCE_APPLICATION_ARB, GL_DEBUG_TYPE_MARKER, GL_DONT_CARE, 0, NULL, GL_FALSE);
+				glDebugMessageControlARB(GL_DEBUG_SOURCE_APPLICATION_ARB, GL_DEBUG_TYPE_PUSH_GROUP, GL_DONT_CARE, 0, NULL, GL_FALSE);
+				glDebugMessageControlARB(GL_DEBUG_SOURCE_APPLICATION_ARB, GL_DEBUG_TYPE_POP_GROUP, GL_DONT_CARE, 0, NULL, GL_FALSE);
+		#ifdef GL_KHR_debug
+				glDebugMessageControlARB(GL_DEBUG_SOURCE_API_ARB, GL_DEBUG_TYPE_OTHER_ARB, GL_DEBUG_SEVERITY_NOTIFICATION, 0, NULL, GL_FALSE);
+		#endif
+				UE_LOG(LogStandaloneRenderer,Warning,TEXT("disabling reporting back of debug groups and markers to the OpenGL debug output callback"));
+			}
+		#endif // GL_ARB_debug_output || GL_KHR_debug
+		}
+	}
+
 #if LINUX_USE_OPENGL_3_2
 	// one Vertex Array Object is reportedly needed for OpenGL 3.2+ core profiles
 	glGenVertexArrays(1, &VertexArrayObject);
diff --git a/Engine/Source/Developer/StandaloneRenderer/Private/OpenGL/SlateOpenGLExtensions.h b/Engine/Source/Developer/StandaloneRenderer/Private/OpenGL/SlateOpenGLExtensions.h
--- a/Engine/Source/Developer/StandaloneRenderer/Private/OpenGL/SlateOpenGLExtensions.h
+++ b/Engine/Source/Developer/StandaloneRenderer/Private/OpenGL/SlateOpenGLExtensions.h
@@ -73,11 +73,13 @@ extern PFNGLISVERTEXARRAYPROC 				glIsVertexArray;
 	EnumMacro(PFNGLENABLEPROC,glEnable) \
 	EnumMacro(PFNGLGENTEXTURESPROC,glGenTextures) \
 	EnumMacro(PFNGLGETERRORPROC,glGetError) \
+	EnumMacro(PFNGLPIXELSTOREIPROC,glPixelStorei) \
 	EnumMacro(PFNGLPOLYGONMODEPROC,glPolygonMode) \
 	EnumMacro(PFNGLSTENCILFUNCPROC,glStencilFunc) \
 	EnumMacro(PFNGLSTENCILMASKPROC,glStencilMask) \
 	EnumMacro(PFNGLSTENCILOPPROC,glStencilOp) \
 	EnumMacro(PFNGLTEXIMAGE2DPROC,glTexImage2D) \
+	EnumMacro(PFNGLTEXSUBIMAGE2DPROC,glTexSubImage2D) \
 	EnumMacro(PFNGLTEXPARAMETERIPROC,glTexParameteri) \
 	EnumMacro(PFNGLVIEWPORTPROC,glViewport) \
 	EnumMacro(PFNGLBLENDEQUATIONPROC,glBlendEquation) \
@@ -110,7 +112,8 @@ extern PFNGLISVERTEXARRAYPROC 				glIsVertexArray;
 	EnumMacro(PFNGLUNIFORM4FVPROC,glUniform4fv) \
 	EnumMacro(PFNGLUNIFORMMATRIX4FVPROC,glUniformMatrix4fv) \
 	EnumMacro(PFNGLVERTEXATTRIBPOINTERPROC,glVertexAttribPointer) \
-	EnumMacro(PFNGLDEBUGMESSAGECALLBACKARBPROC,glDebugMessageCallback) \
+	EnumMacro(PFNGLDEBUGMESSAGECALLBACKARBPROC,glDebugMessageCallbackARB) \
+	EnumMacro(PFNGLDEBUGMESSAGECONTROLARBPROC,glDebugMessageControlARB) \
 	EnumMacro(PFNGLBINDVERTEXARRAYPROC,glBindVertexArray) \
 	EnumMacro(PFNGLDELETEVERTEXARRAYSPROC,glDeleteVertexArrays) \
 	EnumMacro(PFNGLGENVERTEXARRAYSPROC,glGenVertexArrays) \
