OpenGL: fix StandaloneRenderer on Linux and Mac.

diff --git a/Engine/Source/Developer/StandaloneRenderer/Private/OpenGL/SlateOpenGLRenderingPolicy.cpp b/Engine/Source/Developer/StandaloneRenderer/Private/OpenGL/SlateOpenGLRenderingPolicy.cpp
--- a/Engine/Source/Developer/StandaloneRenderer/Private/OpenGL/SlateOpenGLRenderingPolicy.cpp
+++ b/Engine/Source/Developer/StandaloneRenderer/Private/OpenGL/SlateOpenGLRenderingPolicy.cpp
@@ -243,7 +243,7 @@ void FSlateOpenGLRenderingPolicy::DrawElements( const FMatrix& ViewProjectionMat
 
 		glEnableVertexAttribArray(1);
 		Offset = STRUCT_OFFSET( FSlateVertex, Position );
-		glVertexAttribPointer( 1, 2, GL_SHORT, GL_FALSE, Stride, BUFFER_OFFSET(Stride*BaseVertexIndex+Offset) );
+		glVertexAttribPointer( 1, 2, GL_FLOAT, GL_FALSE, Stride, BUFFER_OFFSET(Stride*BaseVertexIndex+Offset) );
 
 		glEnableVertexAttribArray(2);
 		Offset = STRUCT_OFFSET( FSlateVertex, ClipRect );
