Don't set nil content view on Mac when destroying GL context in standalone renderer, as it's illegal in OS X 10.11

diff --git a/Engine/Source/Developer/StandaloneRenderer/Private/Mac/OpenGL/SlateOpenGLContext.cpp b/Engine/Source/Developer/StandaloneRenderer/Private/Mac/OpenGL/SlateOpenGLContext.cpp
--- a/Engine/Source/Developer/StandaloneRenderer/Private/Mac/OpenGL/SlateOpenGLContext.cpp
+++ b/Engine/Source/Developer/StandaloneRenderer/Private/Mac/OpenGL/SlateOpenGLContext.cpp
@@ -168,12 +168,6 @@ void FSlateOpenGLContext::Destroy()
 {
 	if (View)
 	{
-		NSWindow* Window = [View window];
-		if (Window)
-		{
-			[Window setContentView:NULL];
-		}
-		
 		NSOpenGLContext* Current = [NSOpenGLContext currentContext];
 		[Context makeCurrentContext];
 		FSlateCocoaView* SlateView = ((FSlateCocoaView*)View);
