Merge CL#2588905 Mac Standalone Renderer: Don't stomp on the OpenGL context from multiple threads without proper synchronization and avoid invoking Cocoa methods except on the main os thread.
jira: OPP-3633, OPP-3481
Merging using UE4-To-UE4-LauncherDev

diff --git a/Engine/Source/Developer/StandaloneRenderer/Private/Mac/OpenGL/SlateOpenGLContext.cpp b/Engine/Source/Developer/StandaloneRenderer/Private/Mac/OpenGL/SlateOpenGLContext.cpp
--- a/Engine/Source/Developer/StandaloneRenderer/Private/Mac/OpenGL/SlateOpenGLContext.cpp
+++ b/Engine/Source/Developer/StandaloneRenderer/Private/Mac/OpenGL/SlateOpenGLContext.cpp
@@ -11,6 +11,30 @@
 #include <OpenGL/gl.h>
 #include <OpenGL/glext.h>
 
+void LockGLContext(NSOpenGLContext* Context)
+{
+	if (FPlatformMisc::IsRunningOnMavericks())
+	{
+		CGLLockContext([Context CGLContextObj]);
+	}
+	else
+	{
+		[Context lock];
+	}
+}
+
+void UnlockGLContext(NSOpenGLContext* Context)
+{
+	if (FPlatformMisc::IsRunningOnMavericks())
+	{
+		CGLUnlockContext([Context CGLContextObj]);
+	}
+	else
+	{
+		[Context unlock];
+	}
+}
+
 @implementation FSlateOpenGLLayer
 
 - (NSOpenGLPixelFormat *)openGLPixelFormatForDisplayMask:(uint32)Mask
@@ -34,6 +58,16 @@
 	return self;
 }
 
+- (BOOL)canDrawInOpenGLContext:(NSOpenGLContext *)context pixelFormat:(NSOpenGLPixelFormat *)pixelFormat forLayerTime:(CFTimeInterval)timeInterval displayTime:(const CVTimeStamp *)timeStamp
+{
+	BOOL bOK = [super canDrawInOpenGLContext:context pixelFormat:pixelFormat forLayerTime:timeInterval displayTime:timeStamp];
+	if ( bOK && context && (self.Context == context) )
+	{
+		LockGLContext(context);
+	}
+	// TODO: Should we return NO if !(contex && self.Context == context)?
+	return bOK;
+}
 @end
 
 @implementation FSlateCocoaView
@@ -78,6 +112,7 @@
 		glClearColor(0.0f, 0.0f, 0.0f, 0.0f);
 		glClear(GL_COLOR_BUFFER_BIT);
 	}
+	UnlockGLContext(self.Context);
 }
 
 - (BOOL)isOpaque
diff --git a/Engine/Source/Developer/StandaloneRenderer/Private/Mac/OpenGL/SlateOpenGLMac.h b/Engine/Source/Developer/StandaloneRenderer/Private/Mac/OpenGL/SlateOpenGLMac.h
--- a/Engine/Source/Developer/StandaloneRenderer/Private/Mac/OpenGL/SlateOpenGLMac.h
+++ b/Engine/Source/Developer/StandaloneRenderer/Private/Mac/OpenGL/SlateOpenGLMac.h
@@ -10,6 +10,9 @@
 #include <OpenGL/gl.h>
 #include <OpenGL/glext.h>
 
+void LockGLContext(NSOpenGLContext* Context);
+void UnlockGLContext(NSOpenGLContext* Context);
+
 @interface FSlateOpenGLLayer : NSOpenGLLayer
 @property (assign) NSOpenGLContext* Context;
 @property (assign) NSOpenGLPixelFormat* PixelFormat;
diff --git a/Engine/Source/Developer/StandaloneRenderer/Private/Mac/OpenGL/SlateOpenGLViewport.cpp b/Engine/Source/Developer/StandaloneRenderer/Private/Mac/OpenGL/SlateOpenGLViewport.cpp
--- a/Engine/Source/Developer/StandaloneRenderer/Private/Mac/OpenGL/SlateOpenGLViewport.cpp
+++ b/Engine/Source/Developer/StandaloneRenderer/Private/Mac/OpenGL/SlateOpenGLViewport.cpp
@@ -2,7 +2,7 @@
 
 #include "StandaloneRendererPrivate.h"
 #include "OpenGL/SlateOpenGLRenderer.h"
-
+#include "CocoaThread.h"
 #include "SlateOpenGLMac.h"
 
 #include "MacWindow.h"
@@ -32,6 +32,7 @@ void FSlateOpenGLViewport::Destroy()
 
 void FSlateOpenGLViewport::MakeCurrent()
 {
+	LockGLContext(RenderingContext.Context); // NOTE: We assume here that SwapBuffers will always be called after Viewport->MakeCurrent
 	RenderingContext.MakeCurrent();
 	glBindFramebuffer(GL_FRAMEBUFFER, ((FSlateCocoaView*)RenderingContext.View)->Framebuffer);
 }
@@ -40,8 +41,11 @@ void FSlateOpenGLViewport::SwapBuffers()
 {
 	glFlushRenderAPPLE();
 	[(FCocoaWindow*)[RenderingContext.View window] startRendering];
-	[RenderingContext.View setNeedsDisplay:YES];
+	MainThreadCall(^{
+		[RenderingContext.View setNeedsDisplay:YES];
+	}, NSDefaultRunLoopMode, false);
 	glBindFramebuffer(GL_FRAMEBUFFER, ((FSlateCocoaView*)RenderingContext.View)->Framebuffer);
+	UnlockGLContext(RenderingContext.Context);
 }
 
 void FSlateOpenGLViewport::Resize(int32 Width, int32 Height, bool bInFullscreen)
@@ -54,6 +58,7 @@ void FSlateOpenGLViewport::Resize(int32 Width, int32 Height, bool bInFullscreen)
 
 	if (RenderingContext.Context && RenderingContext.View && Width > 0 && Height > 0)
 	{
+		LockGLContext(RenderingContext.Context);
 		((FSlateCocoaView*)RenderingContext.View)->ViewportRect = ViewportRect;
 		
 		GLuint& Framebuffer = ((FSlateCocoaView*)RenderingContext.View)->Framebuffer;
@@ -87,5 +92,6 @@ void FSlateOpenGLViewport::Resize(int32 Width, int32 Height, bool bInFullscreen)
 		glBindFramebuffer(GL_READ_FRAMEBUFFER, CurrentReadFramebuffer);
 
 		[RenderingContext.Context update];
+		UnlockGLContext(RenderingContext.Context);
 	}
 }
diff --git a/Engine/Source/Developer/StandaloneRenderer/Private/OpenGL/SlateOpenGLTextures.cpp b/Engine/Source/Developer/StandaloneRenderer/Private/OpenGL/SlateOpenGLTextures.cpp
--- a/Engine/Source/Developer/StandaloneRenderer/Private/OpenGL/SlateOpenGLTextures.cpp
+++ b/Engine/Source/Developer/StandaloneRenderer/Private/OpenGL/SlateOpenGLTextures.cpp
@@ -3,7 +3,9 @@
 #include "StandaloneRendererPrivate.h"
 #include "SlateOpenGLRenderer.h"
 #include "OpenGL/SlateOpenGLTextures.h"
-
+#if PLATFORM_MAC
+#include "Mac/OpenGL/SlateOpenGLMac.h"
+#endif
 #define USE_DEPRECATED_OPENGL_FUNCTIONALITY			(!PLATFORM_USES_ES2 && !PLATFORM_LINUX)
 
 GLuint FSlateOpenGLTexture::NullTexture = 0;
@@ -63,6 +65,9 @@ void FSlateOpenGLTexture::UpdateTextureThreadSafeRaw(uint32 Width, uint32 Height
 
 void FSlateOpenGLTexture::UpdateTextureRaw(const void* Buffer, const FIntRect& Dirty)
 {
+#if PLATFORM_MAC
+	LockGLContext([NSOpenGLContext currentContext]);
+#endif
 	// Ensure texturing is enabled before setting texture properties
 #if USE_DEPRECATED_OPENGL_FUNCTIONALITY
 	glEnable(GL_TEXTURE_2D);
@@ -93,6 +98,9 @@ void FSlateOpenGLTexture::UpdateTextureRaw(const void* Buffer, const FIntRect& D
 	glTexImage2D(GL_TEXTURE_2D, 0, GL_SRGB8_ALPHA8_EXT, SizeX, SizeY, 0, GL_RGBA, GL_UNSIGNED_INT_8_8_8_8_REV, Buffer);
 #endif
 	CHECK_GL_ERRORS;
+#if PLATFORM_MAC
+	UnlockGLContext([NSOpenGLContext currentContext]);
+#endif
 }
 
 
