[INTEGRATE] [Mac] Lock gl context before deleting it. Prevents occasional crashes when closing windows.

Merging CL#2633925 using UE4-To-UE4-LauncherDev

diff --git a/Engine/Source/Developer/StandaloneRenderer/Private/Mac/OpenGL/SlateOpenGLContext.cpp b/Engine/Source/Developer/StandaloneRenderer/Private/Mac/OpenGL/SlateOpenGLContext.cpp
--- a/Engine/Source/Developer/StandaloneRenderer/Private/Mac/OpenGL/SlateOpenGLContext.cpp
+++ b/Engine/Source/Developer/StandaloneRenderer/Private/Mac/OpenGL/SlateOpenGLContext.cpp
@@ -203,6 +203,7 @@ void FSlateOpenGLContext::Destroy()
 {
 	if (View)
 	{
+		LockGLContext(Context);
 		NSOpenGLContext* Current = [NSOpenGLContext currentContext];
 		[Context makeCurrentContext];
 		FSlateCocoaView* SlateView = ((FSlateCocoaView*)View);
@@ -226,6 +227,7 @@ void FSlateOpenGLContext::Destroy()
 
 		[PixelFormat release];
 		[Context clearDrawable];
+		UnlockGLContext(Context);
 		[Context release];
 		PixelFormat = NULL;
 		Context = NULL;
