Merging using UE4-To-UE4-LauncherDev

diff --git a/Engine/Source/Developer/StandaloneRenderer/Private/OpenGL/SlateOpenGLRenderingPolicy.h b/Engine/Source/Developer/StandaloneRenderer/Private/OpenGL/SlateOpenGLRenderingPolicy.h
--- a/Engine/Source/Developer/StandaloneRenderer/Private/OpenGL/SlateOpenGLRenderingPolicy.h
+++ b/Engine/Source/Developer/StandaloneRenderer/Private/OpenGL/SlateOpenGLRenderingPolicy.h
@@ -33,7 +33,7 @@ public:
 	 */
 	virtual TSharedRef<FSlateFontCache> GetFontCache() override { return FontCache.ToSharedRef(); }
 
-	virtual bool IsVertexColorInLinearSpace() const override { return false; }
+	virtual bool IsVertexColorInLinearSpace() const override { return true; }
 
 	/** 
 	 * Initializes resources if needed
