Reverting bad merge

diff --git a/Engine/Source/Developer/StandaloneRenderer/Private/OpenGL/SlateOpenGLRenderingPolicy.h b/Engine/Source/Developer/StandaloneRenderer/Private/OpenGL/SlateOpenGLRenderingPolicy.h
--- a/Engine/Source/Developer/StandaloneRenderer/Private/OpenGL/SlateOpenGLRenderingPolicy.h
+++ b/Engine/Source/Developer/StandaloneRenderer/Private/OpenGL/SlateOpenGLRenderingPolicy.h
@@ -33,7 +33,7 @@ public:
 	 */
 	virtual TSharedRef<FSlateFontCache> GetFontCache() override { return FontCache.ToSharedRef(); }
 
-	virtual bool IsVertexColorInLinearSpace() const override { return true; }
+	virtual bool IsVertexColorInLinearSpace() const override { return false; }
 
 	/** 
 	 * Initializes resources if needed
