Fixed an issue where we could get a shader compilation error and then continue executing just because the error didn't return a message.

Merging using UE4-To-UE4-LauncherDev

diff --git a/Engine/Source/Developer/StandaloneRenderer/Private/Windows/D3D/SlateD3DShaders.cpp b/Engine/Source/Developer/StandaloneRenderer/Private/Windows/D3D/SlateD3DShaders.cpp
--- a/Engine/Source/Developer/StandaloneRenderer/Private/Windows/D3D/SlateD3DShaders.cpp
+++ b/Engine/Source/Developer/StandaloneRenderer/Private/Windows/D3D/SlateD3DShaders.cpp
@@ -24,6 +24,10 @@ static void CompileShader( const FString& Filename, const FString& EntryPoint, c
 		{
 			checkf(0, ANSI_TO_TCHAR(ErrorBlob->GetBufferPointer()));
 		}
+		else
+		{
+			checkf(0, TEXT("D3DX11CompileFromFile failed, no error text provided"));
+		}
 	}
 }
 
