Fixed various compilation issues.

diff --git a/Engine/Source/Developer/AutomationDriver/Private/AutomatedApplication.cpp b/Engine/Source/Developer/AutomationDriver/Private/AutomatedApplication.cpp
--- a/Engine/Source/Developer/AutomationDriver/Private/AutomatedApplication.cpp
+++ b/Engine/Source/Developer/AutomationDriver/Private/AutomatedApplication.cpp
@@ -170,6 +170,9 @@ class FAutomatedApplicationImpl
 {
 public:
 
+	virtual ~FAutomatedApplicationImpl()
+	{ }
+
 	virtual void AllowPlatformMessageHandling() override
 	{
 		if (PassThroughMessageHandler.IsValid())
diff --git a/Engine/Source/Developer/StandaloneRenderer/Private/Windows/D3D/SlateD3DShaders.cpp b/Engine/Source/Developer/StandaloneRenderer/Private/Windows/D3D/SlateD3DShaders.cpp
--- a/Engine/Source/Developer/StandaloneRenderer/Private/Windows/D3D/SlateD3DShaders.cpp
+++ b/Engine/Source/Developer/StandaloneRenderer/Private/Windows/D3D/SlateD3DShaders.cpp
@@ -4,6 +4,8 @@
 #include "Windows/D3D/SlateD3DRenderer.h"
 #include "Windows/D3D/SlateD3DRenderingPolicy.h"
 #include "Misc/Paths.h"
+#include "HAL/FileManager.h"
+#include "Misc/FileHelper.h"
 
 // @return pointer to the D3DCompile function
 static pD3DCompile GetD3DCompileFunc()
