Slate Standalone Renderers - Making the shaders more resistant to changes in the draw effect enums by interpreting the flags on the CPU, and setting the states as constants, instead of mirroring the enums in the shader code.  Additionally fixing some places where we were not reseting the scissor rastering state in D3D and OGL, that was causing clipping artifacts because the standalone renderers make assumptions about render state when a new frame is drawn.

diff --git a/Engine/Shaders/StandaloneRenderer/D3D/SlateElementPixelShader.hlsl b/Engine/Shaders/StandaloneRenderer/D3D/SlateElementPixelShader.hlsl
--- a/Engine/Shaders/StandaloneRenderer/D3D/SlateElementPixelShader.hlsl
+++ b/Engine/Shaders/StandaloneRenderer/D3D/SlateElementPixelShader.hlsl
@@ -8,10 +8,6 @@
 #define ESlateShader::Font			2
 #define ESlateShader::LineSegment	3
 
-// Draw effects
-#define DE_Disabled 0x20
-#define DE_IgnoreTextureAlpha 0x04
-
 Texture2D ElementTexture;
 SamplerState ElementTextureSampler;
 
@@ -23,20 +19,18 @@ cbuffer PerFramePSConstants
 
 cbuffer PerElementPSConstants
 {
-	float4 ShaderParams;
-	uint DrawEffects;
-	uint ShaderType;
-	uint Padding[2];
+    uint ShaderType;            //  4 bytes
+    float4 ShaderParams;        // 16 bytes
+    uint IgnoreTextureAlpha;    //	4 byte
+    uint DisableEffect;         //  4 byte
+    uint UNUSED[1];             //  4 bytes
 };
 
-
 struct VertexOut
 {
 	float4 Position : SV_POSITION;
 	float4 Color : COLOR0;
 	float4 TextureCoordinates : TEXCOORD0;
-	float4 ClipOriginAndPos : TEXCOORD1;
-	float4 ClipExtents : TEXCOORD2;
 };
 
 float3 Hue( float H )
@@ -73,7 +67,7 @@ float4 GetColor( VertexOut InVertex, float2 UV )
 	float4 FinalColor;
 	
 	float4 BaseColor = ElementTexture.Sample(ElementTextureSampler, UV );
-	if( ( DrawEffects & DE_IgnoreTextureAlpha ) != 0 )
+	if( IgnoreTextureAlpha != 0 )
 	{
 		BaseColor.a = 1.0f;
 	}
@@ -192,7 +186,7 @@ float4 Main( VertexOut InVertex ) : SV_Target
 	// gamma correct
 	OutColor.rgb = GammaCorrect(OutColor.rgb);
 
-	if( DrawEffects & DE_Disabled )
+    if (DisableEffect)
 	{
 		//desaturate
 		float3 LumCoeffs = float3( 0.3, 0.59, .11 );
diff --git a/Engine/Source/Developer/StandaloneRenderer/Private/OpenGL/SlateOpenGLRenderingPolicy.cpp b/Engine/Source/Developer/StandaloneRenderer/Private/OpenGL/SlateOpenGLRenderingPolicy.cpp
--- a/Engine/Source/Developer/StandaloneRenderer/Private/OpenGL/SlateOpenGLRenderingPolicy.cpp
+++ b/Engine/Source/Developer/StandaloneRenderer/Private/OpenGL/SlateOpenGLRenderingPolicy.cpp
@@ -301,6 +301,7 @@ void FSlateOpenGLRenderingPolicy::DrawElements( const FMatrix& ViewProjectionMat
 	}
 
 	// Disable active textures and shaders
+	glDisable(GL_SCISSOR_TEST);
 	glActiveTexture( GL_TEXTURE0 );
 	glBindTexture( GL_TEXTURE_2D, 0 );
 	glUseProgram(0);
diff --git a/Engine/Source/Developer/StandaloneRenderer/Private/Windows/D3D/SlateD3DRenderingPolicy.cpp b/Engine/Source/Developer/StandaloneRenderer/Private/Windows/D3D/SlateD3DRenderingPolicy.cpp
--- a/Engine/Source/Developer/StandaloneRenderer/Private/Windows/D3D/SlateD3DRenderingPolicy.cpp
+++ b/Engine/Source/Developer/StandaloneRenderer/Private/Windows/D3D/SlateD3DRenderingPolicy.cpp
@@ -165,8 +165,6 @@ void FSlateD3D11RenderingPolicy::UpdateVertexAndIndexBuffers( FSlateBatchData& I
 				// @todo make this better
 				VertexBuffer.ResizeBuffer( NumBytesNeeded + 200*sizeof(FSlateVertex) );
 			}
-
-		
 		}
 
 		if( InBatchData.GetNumBatchedIndices() > 0 )
@@ -181,8 +179,6 @@ void FSlateD3D11RenderingPolicy::UpdateVertexAndIndexBuffers( FSlateBatchData& I
 				// @todo make this better
 				IndexBuffer.ResizeBuffer( NumIndices + 100 );
 			}
-
-	
 		}
 
 		uint8* VerticesPtr = nullptr;
@@ -203,7 +199,6 @@ void FSlateD3D11RenderingPolicy::UpdateVertexAndIndexBuffers( FSlateBatchData& I
 			VertexBuffer.Unlock();
 			IndexBuffer.Unlock();
 		}
-
 	}
 }
 
@@ -289,7 +284,6 @@ void FSlateD3D11RenderingPolicy::DrawElements( const FMatrix& ViewProjectionMatr
 				GD3DDeviceContext->RSSetState(NormalRasterState);
 			}
 		}
-	
 
 		PixelShader->SetShaderType( RenderBatch.ShaderType );
 
@@ -327,8 +321,11 @@ void FSlateD3D11RenderingPolicy::DrawElements( const FMatrix& ViewProjectionMatr
 
 		check(RenderBatch.IndexOffset + RenderBatch.NumIndices <= IndexBuffer.GetMaxNumIndices());
 		GD3DDeviceContext->DrawIndexed( IndexCount, RenderBatch.IndexOffset, RenderBatch.VertexOffset );
-
 	}
+
+	// Reset the Raster state when finished, there are places that are making assumptions about the raster state
+	// that need to be fixed.
+	GD3DDeviceContext->RSSetState(NormalRasterState);
 }
 
 TSharedRef<FSlateShaderResourceManager> FSlateD3D11RenderingPolicy::GetResourceManager() const
diff --git a/Engine/Source/Developer/StandaloneRenderer/Private/Windows/D3D/SlateD3DShaders.cpp b/Engine/Source/Developer/StandaloneRenderer/Private/Windows/D3D/SlateD3DShaders.cpp
--- a/Engine/Source/Developer/StandaloneRenderer/Private/Windows/D3D/SlateD3DShaders.cpp
+++ b/Engine/Source/Developer/StandaloneRenderer/Private/Windows/D3D/SlateD3DShaders.cpp
@@ -440,7 +440,8 @@ void FSlateDefaultPS::SetShaderType( uint32 InShaderType )
 
 void FSlateDefaultPS::SetDrawEffects( ESlateDrawEffect InDrawEffects )
 {
-	PerElementConstants.GetBufferData().DrawEffects = (uint32)InDrawEffects;
+	PerElementConstants.GetBufferData().IgnoreTextureAlpha = (uint32)(InDrawEffects & ESlateDrawEffect::IgnoreTextureAlpha);
+	PerElementConstants.GetBufferData().DisableEffect = (uint32)(InDrawEffects & ESlateDrawEffect::DisabledEffect);
 }
 
 void FSlateDefaultPS::SetShaderParams( const FVector4& InShaderParams )
diff --git a/Engine/Source/Developer/StandaloneRenderer/Private/Windows/D3D/SlateD3DShaders.h b/Engine/Source/Developer/StandaloneRenderer/Private/Windows/D3D/SlateD3DShaders.h
--- a/Engine/Source/Developer/StandaloneRenderer/Private/Windows/D3D/SlateD3DShaders.h
+++ b/Engine/Source/Developer/StandaloneRenderer/Private/Windows/D3D/SlateD3DShaders.h
@@ -174,7 +174,7 @@ public:
 	}
 
 	void SetShaderType( uint32 InShaderType );
-	void SetDrawEffects(ESlateDrawEffect InDrawEffects );
+	void SetDrawEffects( ESlateDrawEffect InDrawEffects );
 	void SetShaderParams( const FVector4& InShaderParams );
 	void SetGammaValues(const FVector2D& InGammaValues);
 protected:
@@ -182,9 +182,11 @@ protected:
 private:
 	MS_ALIGN(16) struct FPerElementConstants
 	{
-		FVector4 ShaderParams;
-		uint32 DrawEffects;
-		uint32 ShaderType;
+		uint32 ShaderType;			//  4 bytes
+		FVector4 ShaderParams;		// 16 bytes
+		uint32 IgnoreTextureAlpha;	//	4 byte
+		uint32 DisableEffect;		//  4 byte
+		uint32 UNUSED[1];			//  4 bytes
 	};
 
 	MS_ALIGN(16) struct FPerFrameConstants
