Fix case of IOS files (pt2)

#rnx

diff --git a/Engine/Source/Developer/StandaloneRenderer/Private/IOS/SlateOpenGLContext.cpp b/Engine/Source/Developer/StandaloneRenderer/Private/IOS/SlateOpenGLContext.cpp
new file mode 100644
--- /dev/null
+++ b/Engine/Source/Developer/StandaloneRenderer/Private/IOS/SlateOpenGLContext.cpp
@@ -0,0 +1,28 @@
+// Copyright 1998-2017 Epic Games, Inc. All Rights Reserved.
+
+#include "StandaloneRendererPrivate.h"
+
+#include "OpenGL/SlateOpenGLRenderer.h"
+
+FSlateOpenGLContext::FSlateOpenGLContext()
+:	WindowHandle(NULL)
+,	Context(NULL)
+{
+}
+
+FSlateOpenGLContext::~FSlateOpenGLContext()
+{
+	Destroy();
+}
+
+void FSlateOpenGLContext::Initialize( void* InWindow, const FSlateOpenGLContext* SharedContext )
+{
+}
+
+void FSlateOpenGLContext::Destroy()
+{
+}
+
+void FSlateOpenGLContext::MakeCurrent()
+{
+}
diff --git a/Engine/Source/Developer/StandaloneRenderer/Private/IOS/SlateOpenGLESView.cpp b/Engine/Source/Developer/StandaloneRenderer/Private/IOS/SlateOpenGLESView.cpp
new file mode 100644
--- /dev/null
+++ b/Engine/Source/Developer/StandaloneRenderer/Private/IOS/SlateOpenGLESView.cpp
@@ -0,0 +1,113 @@
+// Copyright 1998-2017 Epic Games, Inc. All Rights Reserved.
+
+#include "StandaloneRendererPrivate.h"
+#include "IOS/SlateOpenGLESView.h"
+#include "IOSAppDelegate.h"
+#include "IOS/IOSInputInterface.h"
+
+@implementation SlateOpenGLESView
+
+@synthesize Context;
+
+- (id)initWithFrame:(CGRect)FrameRect
+{
+	if ((self = [super initWithFrame:FrameRect]))
+	{
+		self.opaque = NO;
+		self.backgroundColor = [UIColor redColor];
+	}
+	return self;
+}
+
+/**
+ * Pass touch events to the input queue for slate to pull off of, and trigger the debug console.
+ *
+ * @param View The view the event happened in
+ * @param Touches Set of touch events from the OS
+ */
+void HandleSlateAppTouches(UIView* View, NSSet* Touches, TouchType Type)
+{
+	// ignore touches until game is booted
+	SlateOpenGLESView* GLView = static_cast<SlateOpenGLESView*>( View );
+	CGFloat Scale = [[UIScreen mainScreen] scale];
+
+	TArray<TouchInput> TouchesArray;
+	for (UITouch* Touch in Touches)
+	{
+		// get info from the touch
+		CGPoint Loc = [Touch locationInView:View];
+		CGPoint PrevLoc = [Touch previousLocationInView:View];
+
+		//@todo zombie - Get the actual Controller ID, and touch pad index. Also, figure out how to deal with touch phases.
+		TouchInput TouchMessage;
+		TouchMessage.Handle = (int)Touch;
+		TouchMessage.Type = Type;
+		TouchMessage.Position = FVector2D(Loc.x, Loc.y) * Scale;
+		TouchMessage.LastPosition = FVector2D(PrevLoc.x, PrevLoc.y) * Scale;
+		TouchesArray.Add(TouchMessage);
+	}
+
+	FIOSInputInterface::QueueTouchInput(TouchesArray);
+}
+
+/**
+ * Handle the various touch types from the OS
+ *
+ * @param touches Array of touch events
+ * @param event Event information
+ */
+- (void) touchesBegan:(NSSet*)touches withEvent:(UIEvent*)event 
+{
+	HandleSlateAppTouches(self, touches, TouchBegan);
+}
+
+- (void) touchesMoved:(NSSet*)touches withEvent:(UIEvent*)event
+{
+	HandleSlateAppTouches(self, touches, TouchMoved);
+}
+
+- (void) touchesEnded:(NSSet*)touches withEvent:(UIEvent*)event
+{
+	HandleSlateAppTouches(self, touches, TouchEnded);
+}
+
+- (void) touchesCancelled:(NSSet*)touches withEvent:(UIEvent*)event
+{
+	HandleSlateAppTouches(self, touches, TouchEnded);
+}
+
+@end
+
+@implementation SlateOpenGLESViewController
+
+/**
+ * The ViewController was created, so now we need to create our view to be controlled (an EAGLView)
+ */
+- (void) loadView
+{
+	// get the landcape size of the screen
+	CGRect Frame = [[UIScreen mainScreen] bounds];
+	EAGLContext* Context = [[EAGLContext alloc] initWithAPI: kEAGLRenderingAPIOpenGLES2];
+	SlateOpenGLESView* View = [[SlateOpenGLESView alloc] initWithFrame:Frame context:Context];
+	self.view = View;
+
+	// pass ownership to the view
+	View.Context = Context;
+	[Context release];
+	
+	// settings copied from InterfaceBuilder
+	self.wantsFullScreenLayout = YES;
+	self.view.clearsContextBeforeDrawing = NO;
+	self.view.multipleTouchEnabled = NO;
+	((GLKView*)self.view).enableSetNeedsDisplay = YES;
+}
+
+/**
+ * Tell the OS that our view controller can auto-rotate between the two landscape modes
+ */
+- (BOOL)shouldAutorotateToInterfaceOrientation:(UIInterfaceOrientation)interfaceOrientation
+{
+	return YES;
+}
+
+@end
diff --git a/Engine/Source/Developer/StandaloneRenderer/Private/IOS/SlateOpenGLViewport.cpp b/Engine/Source/Developer/StandaloneRenderer/Private/IOS/SlateOpenGLViewport.cpp
new file mode 100644
--- /dev/null
+++ b/Engine/Source/Developer/StandaloneRenderer/Private/IOS/SlateOpenGLViewport.cpp
@@ -0,0 +1,104 @@
+// Copyright 1998-2017 Epic Games, Inc. All Rights Reserved.
+
+#include "StandaloneRendererPrivate.h"
+#include "OpenGL/SlateOpenGLRenderer.h"
+#include "IOS/SlateOpenGLESView.h"
+
+#import <OpenGLES/EAGLDrawable.h>
+#import <QuartzCore/QuartzCore.h>
+
+
+uint GDeviceWidth = 0;
+uint GDeviceHeight = 0;
+
+FSlateOpenGLViewport::FSlateOpenGLViewport()
+:	bFullscreen( true )
+{
+}
+
+void FSlateOpenGLViewport::Initialize( TSharedRef<SWindow> InWindow, const FSlateOpenGLContext& SharedContext )
+{
+	TSharedRef<FGenericWindow> NativeWindow = InWindow->GetNativeWindow().ToSharedRef();
+	RenderingContext.WindowHandle = (UIWindow*)NativeWindow->GetOSWindowHandle();
+
+	FSlateRect FullScreenSize = InWindow->GetFullScreenInfo();
+
+	const int Width = FMath::TruncToInt(FullScreenSize.Right - FullScreenSize.Left);
+	const int Height = FMath::TruncToInt(FullScreenSize.Bottom - FullScreenSize.Top);
+
+	// iOS should always be fullscreen and at the origin
+	InWindow->SetCachedSize( FVector2D(Width, Height) );
+	InWindow->SetCachedScreenPosition( FVector2D(0, 0));
+
+	ProjectionMatrix = CreateProjectionMatrix( Width, Height );
+
+	ViewportRect.Right = Width;
+	ViewportRect.Bottom = Height;
+	ViewportRect.Top = 0;
+	ViewportRect.Left = 0;
+
+	UIWindow* MainWindow = (UIWindow*)RenderingContext.WindowHandle;
+	SlateOpenGLESViewController* RootView = (SlateOpenGLESViewController *)MainWindow.rootViewController;
+	SlateOpenGLESView * View = (SlateOpenGLESView *)[RootView view];
+	RenderingContext.Context = View.context;
+	[EAGLContext setCurrentContext: RenderingContext.Context];
+}
+
+void FSlateOpenGLViewport::Destroy()
+{
+	if( RenderingContext.WindowHandle )
+	{
+		UIWindow* MainWindow = (UIWindow*)RenderingContext.WindowHandle;
+		SlateOpenGLESView * View = (SlateOpenGLESView *)[MainWindow.rootViewController view];
+		[View removeFromSuperview];
+		RenderingContext.WindowHandle = NULL;
+		RenderingContext.Context = NULL; // This is released in content view's destructor
+	}
+}
+
+void FSlateOpenGLViewport::MakeCurrent()
+{
+	UIWindow* MainWindow = (UIWindow*)RenderingContext.WindowHandle;
+	SlateOpenGLESView * View = (SlateOpenGLESView *)[MainWindow.rootViewController view];
+	[View bindDrawable];
+}
+
+void FSlateOpenGLViewport::SwapBuffers()
+{
+	UIWindow* MainWindow = (UIWindow*)RenderingContext.WindowHandle;
+	SlateOpenGLESView * View = (SlateOpenGLESView *)[MainWindow.rootViewController view];
+	[View display];
+}
+
+void FSlateOpenGLViewport::Resize( int Width, int Height, bool bInFullscreen )
+{
+	UIWindow* MainWindow = (UIWindow*)RenderingContext.WindowHandle;
+	SlateOpenGLESView * View = (SlateOpenGLESView *)[MainWindow.rootViewController view];
+
+	int ViewportWidth = Width;
+	int ViewportHeight = Height;
+
+	if(bInFullscreen)
+	{
+		CGRect Frame = [[UIScreen mainScreen] bounds];
+		CGFloat Scale = [[UIScreen mainScreen] scale];
+
+		IOSAppDelegate* AppDelegate = (IOSAppDelegate*)[[UIApplication sharedApplication] delegate];
+		if (!AppDelegate.bDeviceInPortraitMode)
+		{
+			Swap<float>(Frame.size.width, Frame.size.height);
+		}
+
+		ViewportWidth = FMath::TruncToInt(Frame.size.width * Scale);
+		ViewportHeight = FMath::TruncToInt(Frame.size.height * Scale);
+	}
+
+	View.frame = CGRectMake(0, 0, ViewportWidth, ViewportHeight);
+
+	ProjectionMatrix = CreateProjectionMatrix( ViewportWidth, ViewportHeight );
+
+	ViewportRect.Right = ViewportWidth;
+	ViewportRect.Bottom = ViewportHeight;
+	ViewportRect.Top = 0;
+	ViewportRect.Left = 0;
+}
