Fixed CrashReportClient window size in high DPI mode and moved NSWindow calls in Mac SlateOpenGLContext to the main thread to solve issues on startup

diff --git a/Engine/Source/Developer/StandaloneRenderer/Private/Mac/OpenGL/SlateOpenGLContext.cpp b/Engine/Source/Developer/StandaloneRenderer/Private/Mac/OpenGL/SlateOpenGLContext.cpp
--- a/Engine/Source/Developer/StandaloneRenderer/Private/Mac/OpenGL/SlateOpenGLContext.cpp
+++ b/Engine/Source/Developer/StandaloneRenderer/Private/Mac/OpenGL/SlateOpenGLContext.cpp
@@ -10,6 +10,7 @@
 #include "Mac/MacTextInputMethodSystem.h"
 #include "Mac/MacPlatformApplicationMisc.h"
 #include "Mac/CocoaTextView.h"
+#include "Mac/CocoaThread.h"
 #include <OpenGL/gl.h>
 #include <OpenGL/glext.h>
 
@@ -193,22 +194,24 @@ void FSlateOpenGLContext::Initialize(void* InWindow, const FSlateOpenGLContext*
 			[View setWantsBestResolutionOpenGLSurface:YES];
 		}
 
-		if (FPlatformMisc::IsRunningOnMavericks() && ([Window styleMask] & NSWindowStyleMaskTexturedBackground))
-		{
-			NSView* SuperView = [[Window contentView] superview];
-			[SuperView addSubview:View];
-			[SuperView setWantsLayer:YES];
-			[SuperView addSubview:[Window standardWindowButton:NSWindowCloseButton]];
-			[SuperView addSubview:[Window standardWindowButton:NSWindowMiniaturizeButton]];
-			[SuperView addSubview:[Window standardWindowButton:NSWindowZoomButton]];
-		}
-		else
-		{
-			[View setWantsLayer:YES];
-			[Window setContentView:View];
-		}
-
-		[[Window standardWindowButton:NSWindowCloseButton] setAction:@selector(performClose:)];
+		MainThreadCall(^{
+			if (FPlatformMisc::IsRunningOnMavericks() && ([Window styleMask] & NSWindowStyleMaskTexturedBackground))
+			{
+				NSView* SuperView = [[Window contentView] superview];
+				[SuperView addSubview:View];
+				[SuperView setWantsLayer:YES];
+				[SuperView addSubview:[Window standardWindowButton:NSWindowCloseButton]];
+				[SuperView addSubview:[Window standardWindowButton:NSWindowMiniaturizeButton]];
+				[SuperView addSubview:[Window standardWindowButton:NSWindowZoomButton]];
+			}
+			else
+			{
+				[View setWantsLayer:YES];
+				[Window setContentView:View];
+			}
+
+			[[Window standardWindowButton:NSWindowCloseButton] setAction:@selector(performClose:)];
+		}, NSDefaultRunLoopMode, true);
 
 		View.layer.magnificationFilter = kCAFilterNearest;
 		View.layer.minificationFilter = kCAFilterNearest;
diff --git a/Engine/Source/Programs/CrashReportClient/Private/CrashReportClientApp.cpp b/Engine/Source/Programs/CrashReportClient/Private/CrashReportClientApp.cpp
--- a/Engine/Source/Programs/CrashReportClient/Private/CrashReportClientApp.cpp
+++ b/Engine/Source/Programs/CrashReportClient/Private/CrashReportClientApp.cpp
@@ -245,13 +245,11 @@ bool RunWithUI(FPlatformErrorReport ErrorReport)
 	// open up the app window	
 	TSharedRef<SCrashReportClient> ClientControl = SNew(SCrashReportClient, CrashReportClient);
 
-	const FSlateRect WorkArea = FSlateApplicationBase::Get().GetPreferredWorkArea();
-
 	auto Window = FSlateApplication::Get().AddWindow(
 		SNew(SWindow)
 		.Title(NSLOCTEXT("CrashReportClient", "CrashReportClientAppName", "Unreal Engine 4 Crash Reporter"))
 		.HasCloseButton(FCrashReportClientConfig::Get().IsAllowedToCloseWithoutSending())
-		.ClientSize(InitialWindowDimensions * FPlatformApplicationMisc::GetDPIScaleFactorAtPoint(WorkArea.Left, WorkArea.Top))
+		.ClientSize(InitialWindowDimensions)
 		[
 			ClientControl
 		]);
