Fix slate viewer

diff --git a/Engine/Source/Developer/StandaloneRenderer/Private/Windows/D3D/SlateD3DTextures.cpp b/Engine/Source/Developer/StandaloneRenderer/Private/Windows/D3D/SlateD3DTextures.cpp
--- a/Engine/Source/Developer/StandaloneRenderer/Private/Windows/D3D/SlateD3DTextures.cpp
+++ b/Engine/Source/Developer/StandaloneRenderer/Private/Windows/D3D/SlateD3DTextures.cpp
@@ -164,7 +164,7 @@ void FSlateD3DTexture::UpdateTextureThreadSafeWithTextureData(FSlateTextureData*
 
 
 FSlateTextureAtlasD3D::FSlateTextureAtlasD3D( uint32 Width, uint32 Height, uint32 StrideBytes, ESlateTextureAtlasPaddingStyle PaddingStyle )
-	: FSlateTextureAtlas( Width, Height, StrideBytes, PaddingStyle )
+	: FSlateTextureAtlas( Width, Height, StrideBytes, PaddingStyle, true )
 	, AtlasTexture( new FSlateD3DTexture( Width, Height ) )
 {
 }
