#UE4 Fix for Linux Tools

diff --git a/Engine/Source/Developer/StandaloneRenderer/Private/OpenGL/SlateOpenGLTextures.cpp b/Engine/Source/Developer/StandaloneRenderer/Private/OpenGL/SlateOpenGLTextures.cpp
--- a/Engine/Source/Developer/StandaloneRenderer/Private/OpenGL/SlateOpenGLTextures.cpp
+++ b/Engine/Source/Developer/StandaloneRenderer/Private/OpenGL/SlateOpenGLTextures.cpp
@@ -183,7 +183,11 @@ GLint FSlateFontTextureOpenGL::GetGLTextureInternalFormat() const
 		return GL_RED;
 #endif // USE_DEPRECATED_OPENGL_FUNCTIONALITY
 	}
+#if !PLATFORM_USES_ES2
+	return GL_SRGB8_ALPHA8;
+#else
 	return GL_SRGB8_ALPHA8_EXT;
+#endif
 }
 
 GLint FSlateFontTextureOpenGL::GetGLTextureFormat() const
