Fixed CrashReportClient freezing when using Copy command from right-click context menu

diff --git a/Engine/Source/Developer/StandaloneRenderer/Private/OpenGL/SlateOpenGLRenderer.cpp b/Engine/Source/Developer/StandaloneRenderer/Private/OpenGL/SlateOpenGLRenderer.cpp
--- a/Engine/Source/Developer/StandaloneRenderer/Private/OpenGL/SlateOpenGLRenderer.cpp
+++ b/Engine/Source/Developer/StandaloneRenderer/Private/OpenGL/SlateOpenGLRenderer.cpp
@@ -130,45 +130,48 @@ void FSlateOpenGLRenderer::DrawWindows( FSlateDrawBuffer& InWindowDrawBuffer )
 			SWindow* WindowToDraw = ElementList.GetRenderWindow();
 
 			const FVector2D WindowSize = WindowToDraw->GetSizeInScreen();
-			FSlateOpenGLViewport* Viewport = WindowToViewportMap.Find( WindowToDraw );
-			check(Viewport);
-		
-			//@todo Slate OpenGL: Move this to ResizeViewport
-			if( WindowSize.X != Viewport->ViewportRect.Right || WindowSize.Y != Viewport->ViewportRect.Bottom )
+			if (WindowSize.X > 0 && WindowSize.Y > 0)
 			{
-				//@todo implement fullscreen
-				const bool bFullscreen = false;
-				Private_ResizeViewport( WindowSize, *Viewport, bFullscreen );
-			}
-			
-			Viewport->MakeCurrent();
+				FSlateOpenGLViewport* Viewport = WindowToViewportMap.Find( WindowToDraw );
+				check(Viewport);
+
+				//@todo Slate OpenGL: Move this to ResizeViewport
+				if( WindowSize.X != Viewport->ViewportRect.Right || WindowSize.Y != Viewport->ViewportRect.Bottom )
+				{
+					//@todo implement fullscreen
+					const bool bFullscreen = false;
+					Private_ResizeViewport( WindowSize, *Viewport, bFullscreen );
+				}
+
+				Viewport->MakeCurrent();
 
-			// Batch elements.  Note that we must set the current viewport before doing this so we have a valid rendering context when calling OpenGL functions
-			ElementBatcher->AddElements(ElementList);
+				// Batch elements.  Note that we must set the current viewport before doing this so we have a valid rendering context when calling OpenGL functions
+				ElementBatcher->AddElements(ElementList);
 
-			// Update the font cache with new text before elements are batched
-			FontCache->UpdateCache();
+				// Update the font cache with new text before elements are batched
+				FontCache->UpdateCache();
 
-			//@ todo Slate: implement for opengl
-			bool bRequiresStencilTest = false;
+				//@ todo Slate: implement for opengl
+				bool bRequiresStencilTest = false;
 
-			ElementBatcher->ResetBatches();
+				ElementBatcher->ResetBatches();
 			
-			FSlateBatchData& BatchData = ElementList.GetBatchData();
+				FSlateBatchData& BatchData = ElementList.GetBatchData();
 
-			RenderingPolicy->BuildRenderingBuffers( BatchData );
+				RenderingPolicy->BuildRenderingBuffers( BatchData );
 
-			check(Viewport);
+				check(Viewport);
 
-			glViewport( Viewport->ViewportRect.Left, Viewport->ViewportRect.Top, Viewport->ViewportRect.Right, Viewport->ViewportRect.Bottom );
+				glViewport( Viewport->ViewportRect.Left, Viewport->ViewportRect.Top, Viewport->ViewportRect.Right, Viewport->ViewportRect.Bottom );
 
-			// Draw all elements
-			RenderingPolicy->DrawElements( ViewMatrix*Viewport->ProjectionMatrix, WindowSize, BatchData.GetRenderBatches() );
+				// Draw all elements
+				RenderingPolicy->DrawElements( ViewMatrix*Viewport->ProjectionMatrix, WindowSize, BatchData.GetRenderBatches() );
 
-			Viewport->SwapBuffers();
+				Viewport->SwapBuffers();
 
-			// All elements have been drawn.  Reset all cached data
-			ElementBatcher->ResetBatches();
+				// All elements have been drawn.  Reset all cached data
+				ElementBatcher->ResetBatches();
+			}
 		}
 	}
 
@@ -180,6 +183,7 @@ void FSlateOpenGLRenderer::DrawWindows( FSlateDrawBuffer& InWindowDrawBuffer )
 }
 
 
+
 /** Called when a window is destroyed to give the renderer a chance to free resources */
 void FSlateOpenGLRenderer::OnWindowDestroyed( const TSharedRef<SWindow>& InWindow )
 {
@@ -315,4 +319,4 @@ int32 FSlateOpenGLRenderer::GetCurrentSceneIndex() const
 void FSlateOpenGLRenderer::ClearScenes() 
 {
 	// This is a no-op
-}
\ No newline at end of file
+}
