Increment the index so we don't spin forever in the OpenGL standalone renderer

diff --git a/Engine/Source/Developer/StandaloneRenderer/Private/OpenGL/SlateOpenGLRenderingPolicy.cpp b/Engine/Source/Developer/StandaloneRenderer/Private/OpenGL/SlateOpenGLRenderingPolicy.cpp
--- a/Engine/Source/Developer/StandaloneRenderer/Private/OpenGL/SlateOpenGLRenderingPolicy.cpp
+++ b/Engine/Source/Developer/StandaloneRenderer/Private/OpenGL/SlateOpenGLRenderingPolicy.cpp
@@ -183,6 +183,8 @@ void FSlateOpenGLRenderingPolicy::DrawElements( const FMatrix& ViewProjectionMat
 	{
 		const FSlateRenderBatch& RenderBatch = RenderBatches[NextRenderBatchIndex];
 
+		NextRenderBatchIndex = RenderBatch.NextBatchIndex;
+
 		const FSlateShaderResource* ShaderResource = RenderBatch.GetShaderResource();
 
 		const ESlateBatchDrawFlag DrawFlags = RenderBatch.GetDrawFlags();
