Fix issue where BatchData.GetRenderBatches() can possibly be empty in the standalone opengl renderer

diff --git a/Engine/Source/Developer/StandaloneRenderer/Private/OpenGL/SlateOpenGLRenderer.cpp b/Engine/Source/Developer/StandaloneRenderer/Private/OpenGL/SlateOpenGLRenderer.cpp
--- a/Engine/Source/Developer/StandaloneRenderer/Private/OpenGL/SlateOpenGLRenderer.cpp
+++ b/Engine/Source/Developer/StandaloneRenderer/Private/OpenGL/SlateOpenGLRenderer.cpp
@@ -164,8 +164,11 @@ void FSlateOpenGLRenderer::DrawWindows( FSlateDrawBuffer& InWindowDrawBuffer )
 
 				glViewport( Viewport->ViewportRect.Left, Viewport->ViewportRect.Top, Viewport->ViewportRect.Right, Viewport->ViewportRect.Bottom );
 
-				// Draw all elements
-				RenderingPolicy->DrawElements( ViewMatrix*Viewport->ProjectionMatrix, WindowSize, BatchData.GetRenderBatches() );
+				if (BatchData.GetRenderBatches().Num() > 0)
+				{
+					// Draw all elements
+					RenderingPolicy->DrawElements( ViewMatrix*Viewport->ProjectionMatrix, WindowSize, BatchData.GetRenderBatches() );
+				}
 
 				Viewport->SwapBuffers();
 
