Fix slate opengl and d3d renderers not restoring gamma correction values after an element requests no gamma correction

diff --git a/Engine/Source/Developer/StandaloneRenderer/Private/OpenGL/SlateOpenGLRenderingPolicy.cpp b/Engine/Source/Developer/StandaloneRenderer/Private/OpenGL/SlateOpenGLRenderingPolicy.cpp
--- a/Engine/Source/Developer/StandaloneRenderer/Private/OpenGL/SlateOpenGLRenderingPolicy.cpp
+++ b/Engine/Source/Developer/StandaloneRenderer/Private/OpenGL/SlateOpenGLRenderingPolicy.cpp
@@ -63,7 +63,7 @@ void FSlateOpenGLRenderingPolicy::ConditionalInitializeResources()
 
 		// Create a default texture.
 		check( WhiteTexture == NULL );
-		WhiteTexture = TextureManager->CreateColorTexture( TEXT("DefaultWhite"), FColor::White );
+		WhiteTexture = TextureManager->CreateColorTexture( TEXT("DefaultWhite"), FColor::White )->GetSlateResource();
 
 		bIsInitialized = true;
 	}
@@ -214,9 +214,13 @@ void FSlateOpenGLRenderingPolicy::DrawElements( const FMatrix& ViewProjectionMat
 		{
 			ElementProgram.SetGammaValues(FVector2D(1.0f, 1.0f));
 		}
+		else
+		{
+			ElementProgram.SetGammaValues(FVector2D(1, 1 / 2.2f));
+		}
 
 		ElementProgram.SetShaderType( static_cast<uint8>(RenderBatch.GetShaderType()) );
-		ElementProgram.SetMarginUVs( RenderBatch.GetShaderParams().PixelParams );
+		ElementProgram.SetShaderParams(RenderBatch.GetShaderParams());
 		ElementProgram.SetDrawEffects( RenderBatch.GetDrawEffects() );
 
 		// Disable stenciling and depth testing by default
@@ -244,7 +248,7 @@ void FSlateOpenGLRenderingPolicy::DrawElements( const FMatrix& ViewProjectionMat
 		}
 		else
 		{
-			ElementProgram.SetTexture( WhiteTexture->GetTypedResource(), GL_CLAMP_TO_EDGE, GL_CLAMP_TO_EDGE );
+			ElementProgram.SetTexture( ((FSlateOpenGLTexture*)WhiteTexture)->GetTypedResource(), GL_CLAMP_TO_EDGE, GL_CLAMP_TO_EDGE );
 		}
 
 		check( RenderBatch.GetNumIndices() > 0 );
diff --git a/Engine/Source/Developer/StandaloneRenderer/Private/Windows/D3D/SlateD3DRenderingPolicy.cpp b/Engine/Source/Developer/StandaloneRenderer/Private/Windows/D3D/SlateD3DRenderingPolicy.cpp
--- a/Engine/Source/Developer/StandaloneRenderer/Private/Windows/D3D/SlateD3DRenderingPolicy.cpp
+++ b/Engine/Source/Developer/StandaloneRenderer/Private/Windows/D3D/SlateD3DRenderingPolicy.cpp
@@ -334,6 +334,10 @@ void FSlateD3D11RenderingPolicy::DrawElements(const FMatrix& ViewProjectionMatri
 		{
 			PixelShader->SetGammaValues(FVector2D(1.0f, 1.0f));
 		}
+		else
+		{
+			PixelShader->SetGammaValues(FVector2D(1, 1 / 2.2f));
+		}
 
 		const FSlateClippingState* ClippingState = RenderBatch.GetClippingState();
 		if (ClippingState != LastClippingState)
@@ -391,7 +395,7 @@ void FSlateD3D11RenderingPolicy::DrawElements(const FMatrix& ViewProjectionMatri
 			PixelShader->SetTexture( ((FSlateD3DTexture*)WhiteTexture)->GetTypedResource(), BilinearSamplerState_Clamp );
 		}
 
-		PixelShader->SetShaderParams( ShaderParams.PixelParams );
+		PixelShader->SetShaderParams(ShaderParams);
 		PixelShader->SetDrawEffects( DrawEffects );
 
 		PixelShader->BindParameters();
