Fix standalone renderers crashing on startup

diff --git a/Engine/Source/Developer/StandaloneRenderer/Private/Windows/D3D/SlateD3DShaders.cpp b/Engine/Source/Developer/StandaloneRenderer/Private/Windows/D3D/SlateD3DShaders.cpp
--- a/Engine/Source/Developer/StandaloneRenderer/Private/Windows/D3D/SlateD3DShaders.cpp
+++ b/Engine/Source/Developer/StandaloneRenderer/Private/Windows/D3D/SlateD3DShaders.cpp
@@ -8,6 +8,18 @@
 #include "Misc/FileHelper.h"
 #include "Misc/App.h"
 
+#define DEFINE_GUID_FOR_CURRENT_COMPILER(name, l, w1, w2, b1, b2, b3, b4, b5, b6, b7, b8) \
+	static const GUID name = { l, w1, w2, { b1, b2,  b3,  b4,  b5,  b6,  b7,  b8 } }
+
+DEFINE_GUID_FOR_CURRENT_COMPILER(IID_ID3D11ShaderReflectionForCurrentCompiler, 0x8d536ca1, 0x0cca, 0x4956, 0xa8, 0x37, 0x78, 0x69, 0x63, 0x75, 0x55, 0x84);
+
+typedef HRESULT(WINAPI* pD3DReflect)
+(__in_bcount(SrcDataSize) LPCVOID pSrcData,
+	__in SIZE_T  SrcDataSize,
+	__in  REFIID pInterface,
+	__out void** ppReflector);
+
+
 // @return pointer to the D3DCompile function
 static pD3DCompile GetD3DCompileFunc()
 {
@@ -35,6 +47,34 @@ static pD3DCompile GetD3DCompileFunc()
 	return &D3DCompile;
 }
 
+// @return pointer to the D3DCompile function
+static pD3DReflect GetD3DReflectFunc()
+{
+	// Override default compiler path to newer dll
+	FString CompilerPath = FPaths::EngineDir();
+#if !PLATFORM_64BITS
+	CompilerPath.Append(TEXT("Binaries/ThirdParty/Windows/DirectX/x86/d3dcompiler_47.dll"));
+#else
+	CompilerPath.Append(TEXT("Binaries/ThirdParty/Windows/DirectX/x64/d3dcompiler_47.dll"));
+#endif
+	static bool bHasCompiler = false;
+	static HMODULE CompilerDLL = 0;
+
+	if (bHasCompiler == false)
+	{
+		CompilerDLL = LoadLibrary(*CompilerPath);
+	}
+
+	if (CompilerDLL)
+	{
+		return (pD3DReflect)(void*)GetProcAddress(CompilerDLL, "D3DReflect");
+	}
+
+	// D3D SDK we compiled with (usually D3DCompiler_43.dll from windows folder)
+	return &D3DReflect;
+}
+
+
 class StandaloneD3DIncluder final : public ID3DInclude
 {
 	public:
@@ -177,8 +217,10 @@ void FSlateD3DVS::Create( const FString& Filename, const FString& EntryPoint, co
 			GEncounteredCriticalD3DDeviceError = true;
 		}
 
+		static pD3DReflect D3DReflectFunc = GetD3DReflectFunc();
+
 		TRefCountPtr<ID3D11ShaderReflection> Reflector;
-		Hr = D3DReflect(Blob->GetBufferPointer(), Blob->GetBufferSize(), IID_ID3D11ShaderReflection, (void**)Reflector.GetInitReference());
+		Hr = D3DReflectFunc(Blob->GetBufferPointer(), Blob->GetBufferSize(), IID_ID3D11ShaderReflectionForCurrentCompiler, (void**)Reflector.GetInitReference());
 		if (FAILED(Hr))
 		{
 			LogSlateD3DRendererFailure(TEXT("FSlateD3DVS::Create() - D3DReflect"), Hr);
@@ -301,8 +343,10 @@ void FSlateD3DPS::Create( const FString& Filename, const FString& EntryPoint, co
 			return;
 		}
 
+		static pD3DReflect D3DReflectFunc = GetD3DReflectFunc();
+
 		TRefCountPtr<ID3D11ShaderReflection> Reflector;
-		Hr = D3DReflect(Blob->GetBufferPointer(), Blob->GetBufferSize(), IID_ID3D11ShaderReflection, (void**)Reflector.GetInitReference());
+		Hr = D3DReflectFunc(Blob->GetBufferPointer(), Blob->GetBufferSize(), IID_ID3D11ShaderReflectionForCurrentCompiler, (void**)Reflector.GetInitReference());
 		if (FAILED(Hr))
 		{
 			LogSlateD3DRendererFailure(TEXT("FSlateD3DPS::Create() - D3DReflect"), Hr);
