Don't cache the pointer to D3DReflect in a static variable, as it seems to cause heartburn for LiveCoding.

Also, add an early return in FSlateD3DVS::Create in case of an error, so it doesn't try to call GetShaderBindings on a null pointer.

#rnx

diff --git a/Engine/Source/Developer/StandaloneRenderer/Private/Windows/D3D/SlateD3DShaders.cpp b/Engine/Source/Developer/StandaloneRenderer/Private/Windows/D3D/SlateD3DShaders.cpp
--- a/Engine/Source/Developer/StandaloneRenderer/Private/Windows/D3D/SlateD3DShaders.cpp
+++ b/Engine/Source/Developer/StandaloneRenderer/Private/Windows/D3D/SlateD3DShaders.cpp
@@ -208,7 +208,7 @@ void FSlateD3DVS::Create( const FString& Filename, const FString& EntryPoint, co
 			GEncounteredCriticalD3DDeviceError = true;
 		}
 
-		static pD3DReflect D3DReflectFunc = GetD3DReflectFunc();
+		pD3DReflect D3DReflectFunc = GetD3DReflectFunc();
 
 		TRefCountPtr<ID3D11ShaderReflection> Reflector;
 		Hr = D3DReflectFunc(Blob->GetBufferPointer(), Blob->GetBufferSize(), IID_ID3D11ShaderReflectionForCurrentCompiler, (void**)Reflector.GetInitReference());
@@ -216,6 +216,7 @@ void FSlateD3DVS::Create( const FString& Filename, const FString& EntryPoint, co
 		{
 			LogSlateD3DRendererFailure(TEXT("FSlateD3DVS::Create() - D3DReflect"), Hr);
 			GEncounteredCriticalD3DDeviceError = true;
+			return;
 		}
 
 		GetShaderBindings(Reflector, ShaderBindings);
@@ -279,7 +280,7 @@ void FSlateD3DPS::Create( const FString& Filename, const FString& EntryPoint, co
 			return;
 		}
 
-		static pD3DReflect D3DReflectFunc = GetD3DReflectFunc();
+		pD3DReflect D3DReflectFunc = GetD3DReflectFunc();
 
 		TRefCountPtr<ID3D11ShaderReflection> Reflector;
 		Hr = D3DReflectFunc(Blob->GetBufferPointer(), Blob->GetBufferSize(), IID_ID3D11ShaderReflectionForCurrentCompiler, (void**)Reflector.GetInitReference());
