Improved slate D3D vertex / index buffer growing algorithm to reduce hitches.

diff --git a/Engine/Source/Developer/StandaloneRenderer/Private/Windows/D3D/SlateD3DRenderingPolicy.cpp b/Engine/Source/Developer/StandaloneRenderer/Private/Windows/D3D/SlateD3DRenderingPolicy.cpp
--- a/Engine/Source/Developer/StandaloneRenderer/Private/Windows/D3D/SlateD3DRenderingPolicy.cpp
+++ b/Engine/Source/Developer/StandaloneRenderer/Private/Windows/D3D/SlateD3DRenderingPolicy.cpp
@@ -233,9 +233,7 @@ void FSlateD3D11RenderingPolicy::BuildRenderingBuffers( FSlateBatchData& InBatch
 			if( NumVertices*sizeof(FSlateVertex) > VertexBuffer.GetBufferSize() )
 			{
 				uint32 NumBytesNeeded = NumVertices*sizeof(FSlateVertex);
-				// increase by a static size.
-				// @todo make this better
-				VertexBuffer.ResizeBuffer( NumBytesNeeded + 200*sizeof(FSlateVertex) );
+				VertexBuffer.ResizeBuffer( NumBytesNeeded + (NumVertices / 4) * sizeof(FSlateVertex) );
 			}
 		}
 
@@ -246,9 +244,7 @@ void FSlateD3D11RenderingPolicy::BuildRenderingBuffers( FSlateBatchData& InBatch
 			// resize if needed
 			if( NumIndices > IndexBuffer.GetMaxNumIndices() )
 			{
-				// increase by a static size.
-				// @todo make this better
-				IndexBuffer.ResizeBuffer( NumIndices + 100 );
+				IndexBuffer.ResizeBuffer( NumIndices + (NumIndices / 4) );
 			}
 		}
 
