LWC: Enforce FMatrix44f for StandaloneRenderer viewprojection matrix

diff --git a/Engine/Source/Developer/StandaloneRenderer/Private/Windows/D3D/SlateD3DShaders.cpp b/Engine/Source/Developer/StandaloneRenderer/Private/Windows/D3D/SlateD3DShaders.cpp
--- a/Engine/Source/Developer/StandaloneRenderer/Private/Windows/D3D/SlateD3DShaders.cpp
+++ b/Engine/Source/Developer/StandaloneRenderer/Private/Windows/D3D/SlateD3DShaders.cpp
@@ -400,7 +400,7 @@ FSlateDefaultVS::FSlateDefaultVS()
 
 void FSlateDefaultVS::SetViewProjection( const FMatrix& ViewProjectionMatrix )
 {
-	ConstantBuffer.GetBufferData().ViewProjection = ViewProjectionMatrix;
+	ConstantBuffer.GetBufferData().ViewProjection = (FMatrix44f)ViewProjectionMatrix;	// LWC_TODO: Precision loss
 }
 
 void FSlateDefaultVS::SetShaderParams( const FVector4& InShaderParams )
diff --git a/Engine/Source/Developer/StandaloneRenderer/Private/Windows/D3D/SlateD3DShaders.h b/Engine/Source/Developer/StandaloneRenderer/Private/Windows/D3D/SlateD3DShaders.h
--- a/Engine/Source/Developer/StandaloneRenderer/Private/Windows/D3D/SlateD3DShaders.h
+++ b/Engine/Source/Developer/StandaloneRenderer/Private/Windows/D3D/SlateD3DShaders.h
@@ -141,9 +141,10 @@ private:
 	
 	struct FPerElementConstants
 	{
-		FMatrix ViewProjection;
+		FMatrix44f ViewProjection;
 		FVector4 VertexShaderParams;
 	};
+	static_assert(sizeof(FPerElementConstants) == sizeof(float) * 20, "Size mismatch");	// LWC_TODO: Fix for FVector4f
 
 	FSlateD3DConstantBuffer<FPerElementConstants> ConstantBuffer;
 	TSlateD3DTypedShaderParameter<ID3D11Buffer>* Constants;
@@ -176,11 +177,13 @@ private:
 		uint32 DisableEffect;		//  4 bytes
 		uint32 UNUSED[1];			//  4 bytes
 	};
+	static_assert(sizeof(FPerElementConstants) == sizeof(float) * 8 + sizeof(uint32) * 4, "Size mismatch");	// LWC_TODO: Fix for FVector4f
 
 	MS_ALIGN(16) struct FPerFrameConstants
 	{
 		FVector2D GammaValues;
 	};
+	static_assert(sizeof(FVector2D) == sizeof(float) * 2, "Size mismatch");	// LWC_TODO: Fix for FVector2f
 
 	FSlateD3DConstantBuffer<FPerElementConstants> PerElementConstants;
 	FSlateD3DConstantBuffer<FPerFrameConstants> PerFrameConstants;
