Find the best available D3D device for standalone Slate application. Use similar options as RHI D3D11.

diff --git a/Engine/Source/Developer/StandaloneRenderer/Private/Windows/D3D/SlateD3DFindBestDevice.cpp b/Engine/Source/Developer/StandaloneRenderer/Private/Windows/D3D/SlateD3DFindBestDevice.cpp
new file mode 100644
--- /dev/null
+++ b/Engine/Source/Developer/StandaloneRenderer/Private/Windows/D3D/SlateD3DFindBestDevice.cpp
@@ -0,0 +1,317 @@
+// Copyright Epic Games, Inc. All Rights Reserved.
+
+#include "Windows/D3D/SlateD3DFindBestDevice.h"
+
+#include "Misc/CommandLine.h"
+#include "StandaloneRendererPlatformHeaders.h"
+#include "StandaloneRendererLog.h"
+#include "Windows/AllowWindowsPlatformTypes.h"
+#include "Windows/D3D/SlateD3dRenderer.h"
+#include "Windows/WindowsPlatformCrashContext.h"
+#include <delayimp.h>
+THIRD_PARTY_INCLUDES_START
+//#include <d3d11_2.h>
+#include <dxgi1_6.h>
+#include <dxgidebug.h>
+THIRD_PARTY_INCLUDES_END
+#include "Microsoft/HideMicrosoftPlatformTypes.h"
+
+namespace UE::StandaloneRenderer::D3D
+{
+namespace Private
+{
+
+bool IsDelayLoadException(PEXCEPTION_POINTERS ExceptionPointers)
+{
+	switch (ExceptionPointers->ExceptionRecord->ExceptionCode)
+	{
+	case VcppException(ERROR_SEVERITY_ERROR, ERROR_MOD_NOT_FOUND):
+	case VcppException(ERROR_SEVERITY_ERROR, ERROR_PROC_NOT_FOUND):
+		return EXCEPTION_EXECUTE_HANDLER;
+	default:
+		return EXCEPTION_CONTINUE_SEARCH;
+	}
+}
+
+} //namespace Private
+
+
+FFindBestDevice::FFindBestDevice()
+{
+	ChosenFeatureLevel = D3D_FEATURE_LEVEL_1_0_CORE;
+	GpuPreference = DXGI_GPU_PREFERENCE_UNSPECIFIED;
+	LoadSettings();
+	SelectAdapter();
+}
+
+FFindBestDevice::~FFindBestDevice()
+{}
+
+bool FFindBestDevice::IsValid() const
+{
+	return ChosenAdapter.IsValid() && ChosenFeatureLevel != D3D_FEATURE_LEVEL_1_0_CORE;
+}
+
+bool FFindBestDevice::CreateDevice(ID3D11Device** Device, ID3D11DeviceContext** DeviceContext)
+{
+	if (!IsValid())
+	{
+		return false;
+	}
+
+	// Init D3D
+	D3D_DRIVER_TYPE DriverType = D3D_DRIVER_TYPE_HARDWARE;
+
+	uint32 DeviceFlags = D3D11_CREATE_DEVICE_SINGLETHREADED;
+	if (bDebug)
+	{
+		DeviceFlags |= D3D11_CREATE_DEVICE_DEBUG;
+	}
+
+	D3D_FEATURE_LEVEL CreatedFeatureLevel = D3D_FEATURE_LEVEL_1_0_CORE;
+	HRESULT Hr = D3D11CreateDevice(ChosenAdapter, D3D_DRIVER_TYPE_UNKNOWN, nullptr, DeviceFlags, &ChosenFeatureLevel, 1, D3D11_SDK_VERSION, Device, &CreatedFeatureLevel, DeviceContext);
+	if (FAILED(Hr))
+	{
+		LogSlateD3DRendererFailure(TEXT("FSlateD3DRenderer::CreateDevice() - D3D11CreateDevice"), Hr);
+	}
+	check(CreatedFeatureLevel == ChosenFeatureLevel);
+	return SUCCEEDED(Hr);
+}
+
+const TCHAR* FFindBestDevice::GetFeatureLevelString(D3D_FEATURE_LEVEL FeatureLevel)
+{
+	switch (FeatureLevel)
+	{
+	case D3D_FEATURE_LEVEL_9_1:		return TEXT("9_1");
+	case D3D_FEATURE_LEVEL_9_2:		return TEXT("9_2");
+	case D3D_FEATURE_LEVEL_9_3:		return TEXT("9_3");
+	case D3D_FEATURE_LEVEL_10_0:	return TEXT("10_0");
+	case D3D_FEATURE_LEVEL_10_1:	return TEXT("10_1");
+	case D3D_FEATURE_LEVEL_11_0:	return TEXT("11_0");
+	case D3D_FEATURE_LEVEL_11_1:	return TEXT("11_1");
+	case D3D_FEATURE_LEVEL_12_0:	return TEXT("12_0");
+	case D3D_FEATURE_LEVEL_12_1:	return TEXT("12_1");
+	case D3D_FEATURE_LEVEL_12_2:	return TEXT("12_2");
+	}
+	return TEXT("X_X");
+}
+
+void FFindBestDevice::LoadSettings()
+{
+	FParse::Value(FCommandLine::Get(), TEXT("-ExplicitAdapterValue="), ExplicitAdapterValue);
+	FParse::Value(FCommandLine::Get(), TEXT("-PreferedAdapterVendor="), PreferedAdapterVendor);
+	bAllowSoftwareRendering = FParse::Param(FCommandLine::Get(), TEXT("AllowSoftwareRendering"));
+	bPreferedMinimalPower = FParse::Param(FCommandLine::Get(), TEXT("PreferedMinimalPower"));
+	bPreferedHighPerformance = FParse::Param(FCommandLine::Get(), TEXT("PreferedHighPerformance"));
+	bDebug = FParse::Param(FCommandLine::Get(), TEXT("d3ddebug"));
+
+	CreateDXGIFactory1(__uuidof(IDXGIFactory1), (void**)DXGIFactory1.GetInitReference());
+	if (DXGIFactory1)
+	{
+		DXGIFactory1->QueryInterface(__uuidof(IDXGIFactory6), (void**)DXGIFactory6.GetInitReference());
+	}
+
+	if (bPreferedMinimalPower)
+	{
+		GpuPreference = DXGI_GPU_PREFERENCE_MINIMUM_POWER;
+	}
+	else
+	{
+		GpuPreference = DXGI_GPU_PREFERENCE_HIGH_PERFORMANCE;
+	}
+
+	UE_LOG(LogStandaloneRenderer, Log, TEXT("D3D11 adapters settings:"));
+	if (ExplicitAdapterValue >= 0)
+	{
+		UE_LOG(LogStandaloneRenderer, Log, TEXT("ExplicitAdapterValue=%d"), ExplicitAdapterValue);
+	}
+	if (PreferedAdapterVendor >= 0)
+	{
+		UE_LOG(LogStandaloneRenderer, Log, TEXT("PreferedAdapterVendor=%d"), PreferedAdapterVendor);
+	}
+	if (bAllowSoftwareRendering)
+	{
+		UE_LOG(LogStandaloneRenderer, Log, TEXT("bAllowSoftwareRendering=%s"), (bAllowSoftwareRendering ? TEXT("true") : TEXT("false")));
+	}
+	if (bPreferedMinimalPower)
+	{
+		UE_LOG(LogStandaloneRenderer, Log, TEXT("bPreferedMinimalPower=%s"), (bPreferedMinimalPower ? TEXT("true") : TEXT("false")));
+	}
+	else if (bPreferedHighPerformance)
+	{
+		UE_LOG(LogStandaloneRenderer, Log, TEXT("bPreferedHighPerformance=%s"), (bPreferedHighPerformance ? TEXT("true") : TEXT("false")));
+	}
+	if (bDebug)
+	{
+		UE_LOG(LogStandaloneRenderer, Log, TEXT("bPreferedHighPerformance=%s"), (bPreferedHighPerformance ? TEXT("true") : TEXT("false")));
+	}
+}
+
+HRESULT FFindBestDevice::EnumerateAdapters(UINT AdapterIndex, IDXGIAdapter** Adapter)
+{
+	if (!DXGIFactory6 || GpuPreference == DXGI_GPU_PREFERENCE_UNSPECIFIED)
+	{
+		return DXGIFactory1->EnumAdapters(AdapterIndex, Adapter);
+	}
+	else
+	{
+		return DXGIFactory6->EnumAdapterByGpuPreference(AdapterIndex, (DXGI_GPU_PREFERENCE)GpuPreference, __uuidof(IDXGIAdapter), (void**)Adapter);
+	}
+}
+
+bool FFindBestDevice::CreateTesingDevice(IDXGIAdapter* Adapter, D3D_FEATURE_LEVEL& ResultFeatureLevel)
+{
+	ID3D11Device* D3DDevice = nullptr;
+	ID3D11DeviceContext* D3DDeviceContext = nullptr;
+
+	uint32 DeviceFlags = D3D11_CREATE_DEVICE_SINGLETHREADED;
+	if (bDebug)
+	{
+		DeviceFlags |= D3D11_CREATE_DEVICE_DEBUG;
+	}
+
+	const D3D_FEATURE_LEVEL RequestedFeatureLevels[] = { D3D_FEATURE_LEVEL_11_1, D3D_FEATURE_LEVEL_11_0 };
+	const int32 NumAllowedFeatureLevels = UE_ARRAY_COUNT(RequestedFeatureLevels);
+
+	ResultFeatureLevel = D3D_FEATURE_LEVEL_1_0_CORE;
+	__try
+	{
+		HRESULT CreateDeviceResult;
+		CreateDeviceResult = D3D11CreateDevice(Adapter, D3D_DRIVER_TYPE_UNKNOWN, nullptr, DeviceFlags, RequestedFeatureLevels, NumAllowedFeatureLevels, D3D11_SDK_VERSION, &D3DDevice, &ResultFeatureLevel, &D3DDeviceContext);
+		if (SUCCEEDED(CreateDeviceResult))
+		{
+			D3DDevice->Release();
+			D3DDeviceContext->Release();
+			return true;
+		}
+		else
+		{
+			UE_LOG(LogStandaloneRenderer, Log, TEXT("Failed to create device [0x%08X]"), CreateDeviceResult);
+		}
+	}
+	__except (Private::IsDelayLoadException(GetExceptionInformation()))
+	{
+		// We suppress warning C6322: Empty _except block. Appropriate checks are made upon returning. 
+		CA_SUPPRESS(6322);
+	}
+
+	return false;
+}
+
+void FFindBestDevice::SelectAdapter()
+{
+	TRefCountPtr<IDXGIAdapter> FirstDXGIAdapterWithoutIntegratedAdapter;
+	TRefCountPtr<IDXGIAdapter> FirstDXGIAdapterAdapter;
+	D3D_FEATURE_LEVEL FirstFeatureLevelWithoutIntegratedAdapter = D3D_FEATURE_LEVEL_1_0_CORE;
+	D3D_FEATURE_LEVEL FirstFeatureLeveAdapter = D3D_FEATURE_LEVEL_1_0_CORE;
+
+	UE_LOG(LogStandaloneRenderer, Log, TEXT("D3D11 adapters:"));
+
+	// Enumerate the DXGIFactory's adapters.
+	TRefCountPtr<IDXGIAdapter> TestingAdapter;
+	for (uint32 AdapterIndex = 0; EnumerateAdapters(AdapterIndex, TestingAdapter.GetInitReference()) != DXGI_ERROR_NOT_FOUND; ++AdapterIndex)
+	{
+		DXGI_ADAPTER_DESC AdapterDesc;
+		if (HRESULT DescResult = TestingAdapter->GetDesc(&AdapterDesc); FAILED(DescResult))
+		{
+			UE_LOG(LogStandaloneRenderer, Warning, TEXT("Failed to get description for adapter %u."), AdapterIndex);
+		}
+		else
+		{
+			UE_LOG(LogStandaloneRenderer, Log, TEXT("Testing D3D11 adapter: %u. Description: '%s'. VendorId: %04x. DeviceId: %04x.")
+				, AdapterIndex, AdapterDesc.Description, AdapterDesc.VendorId, AdapterDesc.DeviceId);
+		}
+
+		D3D_FEATURE_LEVEL FeatureLevel;
+		if (!CreateTesingDevice(TestingAdapter, FeatureLevel))
+		{
+			UE_LOG(LogStandaloneRenderer, Log, TEXT("  Failed to create test device."), AdapterIndex);
+			continue;
+		}
+
+		UE_LOG(LogStandaloneRenderer, Log, TEXT("  %u. '%s'. Feature level: %s"), AdapterIndex, AdapterDesc.Description, GetFeatureLevelString(FeatureLevel));
+
+		const bool bIsMicrosoft = AdapterDesc.VendorId == 0x1414;
+		const bool bIsSoftware = bIsMicrosoft;
+		const bool bSkipSoftwareAdapter = bIsSoftware && !bAllowSoftwareRendering && ExplicitAdapterValue < 0;
+		const bool bSkipExplicitAdapter = ExplicitAdapterValue >= 0 && AdapterIndex != ExplicitAdapterValue;
+		const bool bSkipAdapter = bSkipSoftwareAdapter || bSkipExplicitAdapter;
+
+		if (bSkipAdapter)
+		{
+			UE_LOG(LogStandaloneRenderer, Log, TEXT("  Skip adapter."));
+			continue;
+		}
+
+
+		bool bIsNonLocalMemoryPresent = false;
+		{
+			TRefCountPtr<IDXGIAdapter3> TempDxgiAdapter3;
+			DXGI_QUERY_VIDEO_MEMORY_INFO NonLocalVideoMemoryInfo;
+			if (SUCCEEDED(TestingAdapter->QueryInterface(_uuidof(IDXGIAdapter3), (void**)TempDxgiAdapter3.GetInitReference())) &&
+				TempDxgiAdapter3.IsValid() && SUCCEEDED(TempDxgiAdapter3->QueryVideoMemoryInfo(0, DXGI_MEMORY_SEGMENT_GROUP_NON_LOCAL, &NonLocalVideoMemoryInfo)))
+			{
+				bIsNonLocalMemoryPresent = NonLocalVideoMemoryInfo.Budget != 0;
+			}
+		}
+		const bool bIsIntegrated = !bIsNonLocalMemoryPresent;
+
+		if (!bIsIntegrated && !FirstDXGIAdapterWithoutIntegratedAdapter.IsValid())
+		{
+			FirstDXGIAdapterWithoutIntegratedAdapter = TestingAdapter;
+			FirstFeatureLevelWithoutIntegratedAdapter = FeatureLevel;
+		}
+		else if (PreferedAdapterVendor == AdapterDesc.VendorId && FirstDXGIAdapterWithoutIntegratedAdapter.IsValid())
+		{
+			FirstDXGIAdapterWithoutIntegratedAdapter = TestingAdapter;
+			FirstFeatureLevelWithoutIntegratedAdapter = FeatureLevel;
+		}
+
+		if (!FirstDXGIAdapterAdapter.IsValid())
+		{
+			FirstDXGIAdapterAdapter = TestingAdapter;
+			FirstFeatureLeveAdapter = FeatureLevel;
+		}
+		else if (PreferedAdapterVendor == AdapterDesc.VendorId && FirstDXGIAdapterAdapter.IsValid())
+		{
+			FirstDXGIAdapterAdapter = TestingAdapter;
+			FirstFeatureLeveAdapter = FeatureLevel;
+		}
+	}
+
+	const bool bFavorNonIntegrated = ExplicitAdapterValue == -1;
+	if (bFavorNonIntegrated)
+	{
+		ChosenAdapter = FirstDXGIAdapterWithoutIntegratedAdapter;
+		ChosenFeatureLevel = FirstFeatureLevelWithoutIntegratedAdapter;
+
+		// We assume Intel is integrated graphics (slower than discrete) than NVIDIA or AMD cards and rather take a different one
+		if (!ChosenAdapter.IsValid())
+		{
+			ChosenAdapter = FirstDXGIAdapterAdapter;
+			ChosenFeatureLevel = FirstFeatureLeveAdapter;
+		}
+	}
+	else
+	{
+		ChosenAdapter = FirstDXGIAdapterAdapter;
+		ChosenFeatureLevel = FirstFeatureLeveAdapter;
+	}
+
+	if (ChosenAdapter.IsValid())
+	{
+		DXGI_ADAPTER_DESC AdapterDesc;
+		HRESULT DescResult = ChosenAdapter->GetDesc(&AdapterDesc);
+		if (FAILED(DescResult))
+		{
+			UE_LOG(LogStandaloneRenderer, Warning, TEXT("Failed to get description for selected adapter."));
+		}
+		else
+		{
+			UE_LOG(LogStandaloneRenderer, Log, TEXT("Selected D3D11 Description: '%s'. VendorId: %04x. DeviceId: %04x.")
+				, AdapterDesc.Description, AdapterDesc.VendorId, AdapterDesc.DeviceId);
+		}
+	}
+}
+
+} //namespace
diff --git a/Engine/Source/Developer/StandaloneRenderer/Private/Windows/D3D/SlateD3DFindBestDevice.h b/Engine/Source/Developer/StandaloneRenderer/Private/Windows/D3D/SlateD3DFindBestDevice.h
new file mode 100644
--- /dev/null
+++ b/Engine/Source/Developer/StandaloneRenderer/Private/Windows/D3D/SlateD3DFindBestDevice.h
@@ -0,0 +1,48 @@
+// Copyright Epic Games, Inc. All Rights Reserved.
+
+#pragma once
+
+#include "CoreMinimal.h"
+#include "Templates/RefCounting.h"
+#include "StandaloneRendererPlatformHeaders.h"
+
+struct IDXGIAdapter;
+struct IDXGIFactory1;
+struct IDXGIFactory6;
+
+namespace UE::StandaloneRenderer::D3D
+{
+
+struct FFindBestDevice
+{
+	FFindBestDevice();
+	~FFindBestDevice();
+
+	bool IsValid() const;
+
+	bool CreateDevice(ID3D11Device** Device, ID3D11DeviceContext** DeviceContext);
+
+	static const TCHAR* GetFeatureLevelString(D3D_FEATURE_LEVEL FeatureLevel);
+
+private:
+	void LoadSettings();
+	void SelectAdapter();
+
+	HRESULT EnumerateAdapters(UINT AdapterIndex, IDXGIAdapter** Adapter);
+	bool CreateTesingDevice(IDXGIAdapter* Adapter, D3D_FEATURE_LEVEL& ResultFeatureLevel);
+
+private:
+	TRefCountPtr<IDXGIAdapter> ChosenAdapter;
+	TRefCountPtr<IDXGIFactory1> DXGIFactory1;
+	TRefCountPtr<IDXGIFactory6> DXGIFactory6;
+	D3D_FEATURE_LEVEL ChosenFeatureLevel;
+	int32 GpuPreference; /*DXGI_GPU_PREFERENCE*/
+	int32 ExplicitAdapterValue = -1;
+	int32 PreferedAdapterVendor = -1;
+	bool bAllowSoftwareRendering = false;
+	bool bPreferedMinimalPower = false;
+	bool bPreferedHighPerformance = false;
+	bool bDebug = false;
+};
+
+}
diff --git a/Engine/Source/Developer/StandaloneRenderer/Private/Windows/D3D/SlateD3DRenderer.cpp b/Engine/Source/Developer/StandaloneRenderer/Private/Windows/D3D/SlateD3DRenderer.cpp
--- a/Engine/Source/Developer/StandaloneRenderer/Private/Windows/D3D/SlateD3DRenderer.cpp
+++ b/Engine/Source/Developer/StandaloneRenderer/Private/Windows/D3D/SlateD3DRenderer.cpp
@@ -3,6 +3,7 @@
 #include "Windows/D3D/SlateD3DRenderer.h"
 #include "Windows/D3D/SlateD3DTextureManager.h"
 #include "Windows/D3D/SlateD3DRenderingPolicy.h"
+#include "Windows/D3D/SlateD3DFindBestDevice.h"
 #include "Rendering/ElementBatcher.h"
 #include "Fonts/FontCache.h"
 #include "Widgets/SWindow.h"
@@ -71,7 +72,7 @@ void LogSlateD3DRendererFailure(const FString& Description, HRESULT Hr)
 	if (Hr == DXGI_ERROR_DEVICE_REMOVED && IsValidRef(GD3DDevice))
 	{
 		HRESULT ReasonHr = GD3DDevice->GetDeviceRemovedReason();
-		UE_LOG(LogStandaloneRenderer, Warning, TEXT("%s Reason: %s [%X]"), *Description, *GetReadableResult(ReasonHr), ReasonHr);
+		UE_LOG(LogStandaloneRenderer, Warning, TEXT("%s Reason: %s [0x%08X]"), *Description, *GetReadableResult(ReasonHr), ReasonHr);
 	}
 }
 
@@ -184,24 +185,14 @@ bool FSlateD3DRenderer::CreateDevice()
 
 	if( !IsValidRef( GD3DDevice ) || !IsValidRef( GD3DDeviceContext ) )
 	{
-		// Init D3D
-		uint32 DeviceCreationFlags = D3D11_CREATE_DEVICE_SINGLETHREADED;
-		D3D_DRIVER_TYPE DriverType = D3D_DRIVER_TYPE_HARDWARE;
-
-		if( FParse::Param( FCommandLine::Get(), TEXT("d3ddebug") ) )
-		{
-			DeviceCreationFlags |= D3D11_CREATE_DEVICE_DEBUG;
-		}
-
-		const D3D_FEATURE_LEVEL FeatureLevels[] = { D3D_FEATURE_LEVEL_11_1, D3D_FEATURE_LEVEL_11_0, D3D_FEATURE_LEVEL_10_1, D3D_FEATURE_LEVEL_10_0, D3D_FEATURE_LEVEL_9_3 };
-		D3D_FEATURE_LEVEL CreatedFeatureLevel;
-		HRESULT Hr = D3D11CreateDevice( NULL, DriverType, NULL, DeviceCreationFlags, FeatureLevels, sizeof(FeatureLevels)/sizeof(D3D_FEATURE_LEVEL), D3D11_SDK_VERSION, GD3DDevice.GetInitReference(), &CreatedFeatureLevel, GD3DDeviceContext.GetInitReference() );
-		
-		if (FAILED(Hr))
+		UE::StandaloneRenderer::D3D::FFindBestDevice BestDeviceFinder;
+		if (BestDeviceFinder.IsValid())
 		{
-			bResult = false;
-			LogSlateD3DRendererFailure(TEXT("FSlateD3DRenderer::CreateDevice() - D3D11CreateDevice"), Hr);
-			GEncounteredCriticalD3DDeviceError = true;
+			bResult = BestDeviceFinder.CreateDevice(GD3DDevice.GetInitReference(), GD3DDeviceContext.GetInitReference());
+			if (!bResult)
+			{
+				GEncounteredCriticalD3DDeviceError = true;
+			}
 		}
 	}
 
