Fixed error handling when creating a Slate vertex shader (FSlateD3DVS::Create).

diff --git a/Engine/Source/Developer/StandaloneRenderer/Private/Windows/D3D/SlateD3DShaders.cpp b/Engine/Source/Developer/StandaloneRenderer/Private/Windows/D3D/SlateD3DShaders.cpp
--- a/Engine/Source/Developer/StandaloneRenderer/Private/Windows/D3D/SlateD3DShaders.cpp
+++ b/Engine/Source/Developer/StandaloneRenderer/Private/Windows/D3D/SlateD3DShaders.cpp
@@ -214,6 +214,7 @@ void FSlateD3DVS::Create( const FString& Filename, const FString& EntryPoint, co
 		{
 			LogSlateD3DRendererFailure(TEXT("FSlateD3DVS::Create() - ID3D11Device::CreateVertexShader"), Hr);
 			GEncounteredCriticalD3DDeviceError = true;
+			return;
 		}
 
 		Hr = GD3DDevice->CreateInputLayout(VertexLayout, VertexLayoutCount, Blob->GetBufferPointer(), Blob->GetBufferSize(), InputLayout.GetInitReference());
@@ -221,6 +222,7 @@ void FSlateD3DVS::Create( const FString& Filename, const FString& EntryPoint, co
 		{
 			LogSlateD3DRendererFailure(TEXT("FSlateD3DVS::Create() - ID3D11Device::CreateInputLayout"), Hr);
 			GEncounteredCriticalD3DDeviceError = true;
+			return;
 		}
 
 		pD3DReflect D3DReflectFunc = GetD3DReflectFunc();
