Add the -nullrhi command option to all standalone slate application.

diff --git a/Engine/Source/Developer/StandaloneRenderer/Private/StandaloneRenderer.cpp b/Engine/Source/Developer/StandaloneRenderer/Private/StandaloneRenderer.cpp
--- a/Engine/Source/Developer/StandaloneRenderer/Private/StandaloneRenderer.cpp
+++ b/Engine/Source/Developer/StandaloneRenderer/Private/StandaloneRenderer.cpp
@@ -1,6 +1,8 @@
 // Copyright Epic Games, Inc. All Rights Reserved.
 
 #include "StandaloneRenderer.h"
+#include "Interfaces/ISlateNullRendererModule.h"
+#include "Misc/App.h"
 #include "Misc/CommandLine.h"
 #include "Modules/ModuleInterface.h"
 #include "Modules/ModuleManager.h"
@@ -23,19 +25,26 @@ TSharedRef<FSlateRenderer> GetStandardStandaloneRenderer()
 {
 	// create a standalone renderer object
 	TSharedPtr<FSlateRenderer> Renderer = NULL;
-#if PLATFORM_WINDOWS
-	bool bUseOpenGL = FParse::Param( FCommandLine::Get(), TEXT("opengl") );
-	if( bUseOpenGL )
+	if (FApp::CanEverRender())
 	{
+#if PLATFORM_WINDOWS
+		bool bUseOpenGL = FParse::Param( FCommandLine::Get(), TEXT("opengl") );
+		if( bUseOpenGL )
+		{
 #endif
-		Renderer = TSharedPtr<FSlateRenderer>( new FSlateOpenGLRenderer( FCoreStyle::Get() ) );
+			Renderer = TSharedPtr<FSlateRenderer>( new FSlateOpenGLRenderer( FCoreStyle::Get() ) );
 #if PLATFORM_WINDOWS
+		}
+		else
+		{
+			Renderer = TSharedPtr<FSlateRenderer>( new FSlateD3DRenderer( FCoreStyle::Get() ) );
+		}
+#endif
 	}
 	else
 	{
-		Renderer = TSharedPtr<FSlateRenderer>( new FSlateD3DRenderer( FCoreStyle::Get() ) );
+		Renderer = FModuleManager::Get().LoadModuleChecked<ISlateNullRendererModule>("SlateNullRenderer").CreateSlateNullRenderer();
 	}
-#endif
 
 	// enforce non-NULL pointer
 	return Renderer.ToSharedRef();
diff --git a/Engine/Source/Developer/StandaloneRenderer/StandaloneRenderer.Build.cs b/Engine/Source/Developer/StandaloneRenderer/StandaloneRenderer.Build.cs
--- a/Engine/Source/Developer/StandaloneRenderer/StandaloneRenderer.Build.cs
+++ b/Engine/Source/Developer/StandaloneRenderer/StandaloneRenderer.Build.cs
@@ -14,6 +14,7 @@ public class StandaloneRenderer : ModuleRules
 				"ImageWrapper",
 				"InputCore",
 				"SlateCore",
+				"SlateNullRenderer",
 			}
 			);
 
diff --git a/Engine/Source/Runtime/SlateNullRenderer/Private/SlateNullRenderer.cpp b/Engine/Source/Runtime/SlateNullRenderer/Private/SlateNullRenderer.cpp
--- a/Engine/Source/Runtime/SlateNullRenderer/Private/SlateNullRenderer.cpp
+++ b/Engine/Source/Runtime/SlateNullRenderer/Private/SlateNullRenderer.cpp
@@ -4,7 +4,9 @@
 #include "Rendering/DrawElements.h"
 #include "Rendering/ShaderResourceManager.h"
 #include "Rendering/SlateDrawBuffer.h"
+#if UE_SLATE_NULL_RENDERER_WITH_ENGINE
 #include "UnrealEngine.h"
+#endif
 
 
 FSlateNullRenderer::FSlateNullRenderer(const TSharedRef<FSlateFontServices>& InSlateFontServices, const TSharedRef<FSlateShaderResourceManager>& InResourceManager)
@@ -136,8 +138,10 @@ void FSlateNullRenderer::ClearScenes()
 
 void FSlateNullRenderer::Sync() const
 {
+#if UE_SLATE_NULL_RENDERER_WITH_ENGINE
 	// Sync game and render thread. Either total sync or allowing one frame lag.
 	static FFrameEndSync FrameEndSync;
 	static auto CVarAllowOneFrameThreadLag = IConsoleManager::Get().FindTConsoleVariableDataInt(TEXT("r.OneFrameThreadLag"));
 	FrameEndSync.Sync(CVarAllowOneFrameThreadLag->GetValueOnAnyThread() != 0);
+#endif
 }
diff --git a/Engine/Source/Runtime/SlateNullRenderer/SlateNullRenderer.Build.cs b/Engine/Source/Runtime/SlateNullRenderer/SlateNullRenderer.Build.cs
--- a/Engine/Source/Runtime/SlateNullRenderer/SlateNullRenderer.Build.cs
+++ b/Engine/Source/Runtime/SlateNullRenderer/SlateNullRenderer.Build.cs
@@ -16,13 +16,17 @@ public class SlateNullRenderer : ModuleRules
 
 		if (Target.bCompileAgainstEngine)
 		{
+			PrivateDefinitions.Add("UE_SLATE_NULL_RENDERER_WITH_ENGINE=1");
 			PrivateDependencyModuleNames.AddRange(
-			new string[] {
-				"Engine",
-				"RenderCore",
-				"RHI"
-			}
-		);
+				new string[] {
+					"Engine",
+					"RenderCore",
+					"RHI"
+				});
+		}
+		else
+		{
+			PrivateDefinitions.Add("UE_SLATE_NULL_RENDERER_WITH_ENGINE=0");
 		}
 	}
 }
