Add overriding implementation of GenerateDynamicImageResource for FSlateD3DRenderer

diff --git a/Engine/Source/Developer/StandaloneRenderer/Private/Windows/D3D/SlateD3DRenderer.cpp b/Engine/Source/Developer/StandaloneRenderer/Private/Windows/D3D/SlateD3DRenderer.cpp
--- a/Engine/Source/Developer/StandaloneRenderer/Private/Windows/D3D/SlateD3DRenderer.cpp
+++ b/Engine/Source/Developer/StandaloneRenderer/Private/Windows/D3D/SlateD3DRenderer.cpp
@@ -404,6 +404,11 @@ bool FSlateD3DRenderer::GenerateDynamicImageResource(FName ResourceName, uint32
 	return Result != nullptr;
 }
 
+bool FSlateD3DRenderer::GenerateDynamicImageResource(FName ResourceName, FSlateTextureDataRef TextureData)
+{
+	return GenerateDynamicImageResource(ResourceName, TextureData->GetWidth(), TextureData->GetHeight(), TextureData->GetRawBytes());
+}
+
 FSlateResourceHandle FSlateD3DRenderer::GetResourceHandle(const FSlateBrush& Brush, FVector2f LocalSize, float DrawScale)
 {
 	if (!TextureManager.IsValid())
diff --git a/Engine/Source/Developer/StandaloneRenderer/Private/Windows/D3D/SlateD3DRenderer.h b/Engine/Source/Developer/StandaloneRenderer/Private/Windows/D3D/SlateD3DRenderer.h
--- a/Engine/Source/Developer/StandaloneRenderer/Private/Windows/D3D/SlateD3DRenderer.h
+++ b/Engine/Source/Developer/StandaloneRenderer/Private/Windows/D3D/SlateD3DRenderer.h
@@ -63,6 +63,7 @@ public:
 	virtual void RestoreSystemResolution(const TSharedRef<SWindow> InWindow) override {}
 	virtual void ReleaseDynamicResource( const FSlateBrush& Brush ) override;
 	virtual bool GenerateDynamicImageResource(FName ResourceName, uint32 Width, uint32 Height, const TArray< uint8 >& Bytes) override;
+	virtual bool GenerateDynamicImageResource(FName ResourceName, FSlateTextureDataRef TextureData) override;
 	virtual FSlateResourceHandle GetResourceHandle(const FSlateBrush& Brush, FVector2f LocalSize, float DrawScale) override;
 	virtual void RemoveDynamicBrushResource( TSharedPtr<FSlateDynamicImageBrush> BrushToRemove ) override;
 	virtual void LoadStyleResources( const ISlateStyle& Style ) override;
