Add overriding implementation of GenerateDynamicImageResource for FSlateOpenGLRenderer

diff --git a/Engine/Source/Developer/StandaloneRenderer/Private/OpenGL/SlateOpenGLRenderer.cpp b/Engine/Source/Developer/StandaloneRenderer/Private/OpenGL/SlateOpenGLRenderer.cpp
--- a/Engine/Source/Developer/StandaloneRenderer/Private/OpenGL/SlateOpenGLRenderer.cpp
+++ b/Engine/Source/Developer/StandaloneRenderer/Private/OpenGL/SlateOpenGLRenderer.cpp
@@ -270,6 +270,12 @@ bool FSlateOpenGLRenderer::GenerateDynamicImageResource(FName ResourceName, uint
 	return TextureManager->CreateDynamicTextureResource(ResourceName, Width, Height, Bytes) != NULL;
 }
 
+bool FSlateOpenGLRenderer::GenerateDynamicImageResource(FName ResourceName, FSlateTextureDataRef TextureData)
+{
+	return GenerateDynamicImageResource(ResourceName, TextureData->GetWidth(), TextureData->GetHeight(), TextureData->GetRawBytes());
+}
+
+
 FSlateResourceHandle FSlateOpenGLRenderer::GetResourceHandle(const FSlateBrush& Brush, FVector2f LocalSize, float DrawScale)
 {
 	return TextureManager->GetResourceHandle(Brush, LocalSize, DrawScale);
diff --git a/Engine/Source/Developer/StandaloneRenderer/Private/OpenGL/SlateOpenGLRenderer.h b/Engine/Source/Developer/StandaloneRenderer/Private/OpenGL/SlateOpenGLRenderer.h
--- a/Engine/Source/Developer/StandaloneRenderer/Private/OpenGL/SlateOpenGLRenderer.h
+++ b/Engine/Source/Developer/StandaloneRenderer/Private/OpenGL/SlateOpenGLRenderer.h
@@ -159,6 +159,7 @@ public:
 	virtual void RestoreSystemResolution(const TSharedRef<SWindow> InWindow) override {}
 	virtual void ReleaseDynamicResource( const FSlateBrush& Brush ) override;
 	virtual bool GenerateDynamicImageResource(FName ResourceName, uint32 Width, uint32 Height, const TArray< uint8 >& Bytes) override;
+	virtual bool GenerateDynamicImageResource(FName ResourceName, FSlateTextureDataRef TextureData) override;
 	virtual FSlateResourceHandle GetResourceHandle(const FSlateBrush& Brush, FVector2f LocalSize, float DrawScale) override;
 	virtual void RemoveDynamicBrushResource( TSharedPtr<FSlateDynamicImageBrush> BrushToRemove ) override;
 	virtual void LoadStyleResources(const ISlateStyle& Style) override;
