PR #11218: UE-191987: FAtlasedTextureSlot from VectorGraphicsCache are not stored in AtlasDebugData

diff --git a/Engine/Source/Developer/StandaloneRenderer/Private/OpenGL/SlateOpenGLTextureManager.cpp b/Engine/Source/Developer/StandaloneRenderer/Private/OpenGL/SlateOpenGLTextureManager.cpp
--- a/Engine/Source/Developer/StandaloneRenderer/Private/OpenGL/SlateOpenGLTextureManager.cpp
+++ b/Engine/Source/Developer/StandaloneRenderer/Private/OpenGL/SlateOpenGLTextureManager.cpp
@@ -468,7 +468,8 @@ FAtlasSlotInfo FSlateOpenGLTextureManager::GetAtlasSlotInfoAtPosition(FIntPoint
 {
 	const FSlateTextureAtlas* Atlas = nullptr;
 
-	if (PrecachedTextureAtlases.IsValidIndex(AtlasIndex))
+	bool bIsPrecachedTextureAtlases = PrecachedTextureAtlases.IsValidIndex(AtlasIndex);
+	if (bIsPrecachedTextureAtlases)
 	{
 		Atlas = PrecachedTextureAtlases[AtlasIndex].Get();
 	}
@@ -477,21 +478,17 @@ FAtlasSlotInfo FSlateOpenGLTextureManager::GetAtlasSlotInfoAtPosition(FIntPoint
 		Atlas = VectorGraphicsCache->GetAtlas(AtlasIndex - PrecachedTextureAtlases.Num());
 	}
 
+	FAtlasSlotInfo NewInfo;
 	if (Atlas)
 	{
-		FAtlasSlotInfo NewInfo;
-
 		const FAtlasedTextureSlot* Slot = Atlas->GetSlotAtPosition(InPosition);
 		if (Slot)
 		{
 			NewInfo.AtlasSlotRect = FSlateRect(FVector2f(Slot->X, Slot->Y), FVector2f(Slot->X + Slot->Width, Slot->Y + Slot->Height));
-
-			NewInfo.TextureName = AtlasDebugData.FindRef(Slot);
-
-			return NewInfo;
+			NewInfo.TextureName = bIsPrecachedTextureAtlases ? AtlasDebugData.FindRef(Slot) : VectorGraphicsCache->GetAtlasDebugData(Slot);
 		}
 	}
 
-	return FAtlasSlotInfo();
+	return NewInfo;
 }
 #endif
\ No newline at end of file
diff --git a/Engine/Source/Developer/StandaloneRenderer/Private/Windows/D3D/SlateD3DTextureManager.cpp b/Engine/Source/Developer/StandaloneRenderer/Private/Windows/D3D/SlateD3DTextureManager.cpp
--- a/Engine/Source/Developer/StandaloneRenderer/Private/Windows/D3D/SlateD3DTextureManager.cpp
+++ b/Engine/Source/Developer/StandaloneRenderer/Private/Windows/D3D/SlateD3DTextureManager.cpp
@@ -98,7 +98,8 @@ FAtlasSlotInfo FSlateD3DTextureManager::GetAtlasSlotInfoAtPosition(FIntPoint InP
 {
 	const FSlateTextureAtlas* Atlas = nullptr;
 
-	if (PrecachedTextureAtlases.IsValidIndex(AtlasIndex))
+	bool bIsPrecachedTextureAtlases = PrecachedTextureAtlases.IsValidIndex(AtlasIndex);
+	if (bIsPrecachedTextureAtlases)
 	{
 		Atlas = PrecachedTextureAtlases[AtlasIndex].Get();
 	}
@@ -107,22 +108,18 @@ FAtlasSlotInfo FSlateD3DTextureManager::GetAtlasSlotInfoAtPosition(FIntPoint InP
 		Atlas = VectorGraphicsCache->GetAtlas(AtlasIndex - PrecachedTextureAtlases.Num());
 	}
 
+	FAtlasSlotInfo NewInfo;
 	if(Atlas)
 	{
-		FAtlasSlotInfo NewInfo;
-
 		const FAtlasedTextureSlot* Slot = Atlas->GetSlotAtPosition(InPosition);
 		if (Slot)
 		{
 			NewInfo.AtlasSlotRect = FSlateRect(FVector2f(Slot->X, Slot->Y), FVector2f(Slot->X + Slot->Width, Slot->Y + Slot->Height));
-
-			NewInfo.TextureName = AtlasDebugData.FindRef(Slot);
-
-			return NewInfo;
+			NewInfo.TextureName = bIsPrecachedTextureAtlases ? AtlasDebugData.FindRef(Slot) : VectorGraphicsCache->GetAtlasDebugData(Slot);
 		}
 	}
 
-	return FAtlasSlotInfo();
+	return NewInfo;
 }
 #endif
 
diff --git a/Engine/Source/Runtime/SlateCore/Private/Rendering/SlateVectorGraphicsCache.cpp b/Engine/Source/Runtime/SlateCore/Private/Rendering/SlateVectorGraphicsCache.cpp
--- a/Engine/Source/Runtime/SlateCore/Private/Rendering/SlateVectorGraphicsCache.cpp
+++ b/Engine/Source/Runtime/SlateCore/Private/Rendering/SlateVectorGraphicsCache.cpp
@@ -149,6 +149,10 @@ void FSlateVectorGraphicsCache::UpdateCache()
 						NewProxy->StartUV = FVector2f((float)(NewSlot->X + Padding) / FoundAtlas->GetWidth(), (float)(NewSlot->Y + Padding) / FoundAtlas->GetHeight());
 						NewProxy->SizeUV = FVector2f((float)(NewSlot->Width - Padding * 2) / FoundAtlas->GetWidth(), (float)(NewSlot->Height - Padding * 2) / FoundAtlas->GetHeight());
 						NewProxy->ActualSize = FIntPoint(PixelSize.X, PixelSize.Y);
+
+#if WITH_ATLAS_DEBUGGING
+						AtlasDebugData.Add(NewSlot, Request.Key.BrushName);
+#endif
 					}
 				}
 				UpdateFlushCounters(0, Atlases.Num(), 0, NonAtlasedTextures.Num());
@@ -204,6 +208,9 @@ void FSlateVectorGraphicsCache::ReleaseResources(bool bWaitForRelease /* = false
 
 void FSlateVectorGraphicsCache::DeleteResources()
 {
+#if WITH_ATLAS_DEBUGGING
+	AtlasDebugData.Empty();
+#endif
 	Atlases.Empty();
 	NonAtlasedTextures.Empty();
 	ResourceMap.Empty();
diff --git a/Engine/Source/Runtime/SlateCore/Public/Rendering/SlateVectorGraphicsCache.h b/Engine/Source/Runtime/SlateCore/Public/Rendering/SlateVectorGraphicsCache.h
--- a/Engine/Source/Runtime/SlateCore/Public/Rendering/SlateVectorGraphicsCache.h
+++ b/Engine/Source/Runtime/SlateCore/Public/Rendering/SlateVectorGraphicsCache.h
@@ -32,6 +32,10 @@ public:
 	bool IsAtlasPageResourceAlphaOnly(const int32 InIndex) const { return false; }
 	const FSlateTextureAtlas* GetAtlas(const int32 InIndex) const { return Atlases[InIndex].Get(); }
 
+#if WITH_ATLAS_DEBUGGING
+	FName GetAtlasDebugData(const FAtlasedTextureSlot* InSlot) const { return AtlasDebugData.FindRef(InSlot);  }
+#endif
+
 private:
 	void FlushCache();
 
@@ -73,6 +77,9 @@ private:
 		{}
 	};
 
+#if WITH_ATLAS_DEBUGGING
+	TMap<const FAtlasedTextureSlot*, FName> AtlasDebugData;
+#endif
 	TMap<FVectorCacheKey, TUniquePtr<FSlateShaderResourceProxy>> ResourceMap;
 	TArray<TUniquePtr<FSlateTextureAtlas>> Atlases;
 	TArray<TUniquePtr<FSlateShaderResource>> NonAtlasedTextures;
diff --git a/Engine/Source/Runtime/SlateRHIRenderer/Private/SlateRHIResourceManager.cpp b/Engine/Source/Runtime/SlateRHIRenderer/Private/SlateRHIResourceManager.cpp
--- a/Engine/Source/Runtime/SlateRHIRenderer/Private/SlateRHIResourceManager.cpp
+++ b/Engine/Source/Runtime/SlateRHIRenderer/Private/SlateRHIResourceManager.cpp
@@ -382,7 +382,8 @@ FAtlasSlotInfo FSlateRHIResourceManager::GetAtlasSlotInfoAtPosition(FIntPoint In
 {
 	const FSlateTextureAtlas* Atlas = nullptr;
 
-	if (PrecachedTextureAtlases.IsValidIndex(AtlasIndex))
+	bool bIsPrecachedTextureAtlases = PrecachedTextureAtlases.IsValidIndex(AtlasIndex);
+	if (bIsPrecachedTextureAtlases)
 	{
 		Atlas = PrecachedTextureAtlases[AtlasIndex];
 	}
@@ -391,22 +392,18 @@ FAtlasSlotInfo FSlateRHIResourceManager::GetAtlasSlotInfoAtPosition(FIntPoint In
 		Atlas = VectorGraphicsCache->GetAtlas(AtlasIndex - PrecachedTextureAtlases.Num());
 	}
 
+	FAtlasSlotInfo NewInfo;
 	if (Atlas)
 	{
-		FAtlasSlotInfo NewInfo;
-
 		const FAtlasedTextureSlot* Slot = Atlas->GetSlotAtPosition(InPosition);
 		if (Slot)
 		{
 			NewInfo.AtlasSlotRect = FSlateRect(FVector2f((float)Slot->X, (float)Slot->Y), FVector2f((float)(Slot->X + Slot->Width), (float)(Slot->Y + Slot->Height)));
-
-			NewInfo.TextureName = AtlasDebugData.FindRef(Slot);
-
-			return NewInfo;
+			NewInfo.TextureName = bIsPrecachedTextureAtlases ? AtlasDebugData.FindRef(Slot) : VectorGraphicsCache->GetAtlasDebugData(Slot);
 		}
 	}
 
-	return FAtlasSlotInfo();
+	return NewInfo;
 }
 #endif
 
