Fixing build issues with older Windows SDK versions.

diff --git a/Engine/Source/Developer/StandaloneRenderer/Private/Windows/D3D/SlateD3DFindBestDevice.cpp b/Engine/Source/Developer/StandaloneRenderer/Private/Windows/D3D/SlateD3DFindBestDevice.cpp
--- a/Engine/Source/Developer/StandaloneRenderer/Private/Windows/D3D/SlateD3DFindBestDevice.cpp
+++ b/Engine/Source/Developer/StandaloneRenderer/Private/Windows/D3D/SlateD3DFindBestDevice.cpp
@@ -80,18 +80,11 @@ bool FFindBestDevice::CreateDevice(ID3D11Device** Device, ID3D11DeviceContext**
 
 const TCHAR* FFindBestDevice::GetFeatureLevelString(D3D_FEATURE_LEVEL FeatureLevel)
 {
+	// Note: this switch only needs the feature levels used inside FFindBestDevice::CreateTesingDevice
 	switch (FeatureLevel)
 	{
-	case D3D_FEATURE_LEVEL_9_1:		return TEXT("9_1");
-	case D3D_FEATURE_LEVEL_9_2:		return TEXT("9_2");
-	case D3D_FEATURE_LEVEL_9_3:		return TEXT("9_3");
-	case D3D_FEATURE_LEVEL_10_0:	return TEXT("10_0");
-	case D3D_FEATURE_LEVEL_10_1:	return TEXT("10_1");
 	case D3D_FEATURE_LEVEL_11_0:	return TEXT("11_0");
 	case D3D_FEATURE_LEVEL_11_1:	return TEXT("11_1");
-	case D3D_FEATURE_LEVEL_12_0:	return TEXT("12_0");
-	case D3D_FEATURE_LEVEL_12_1:	return TEXT("12_1");
-	case D3D_FEATURE_LEVEL_12_2:	return TEXT("12_2");
 	}
 	return TEXT("X_X");
 }
diff --git a/Engine/Source/Developer/Windows/ShaderFormatD3D/ShaderFormatD3D.Build.cs b/Engine/Source/Developer/Windows/ShaderFormatD3D/ShaderFormatD3D.Build.cs
--- a/Engine/Source/Developer/Windows/ShaderFormatD3D/ShaderFormatD3D.Build.cs
+++ b/Engine/Source/Developer/Windows/ShaderFormatD3D/ShaderFormatD3D.Build.cs
@@ -44,7 +44,7 @@ public class ShaderFormatD3D : ModuleRules
 			);
         }
 
-		AddEngineThirdPartyPrivateStaticDependencies(Target, "DX11");
+		AddEngineThirdPartyPrivateStaticDependencies(Target, "AgilitySDK");
 
 		// DXC API headers
 		PublicSystemIncludePaths.Add(Path.Combine(Target.UEThirdPartySourceDirectory, "ShaderConductor", "ShaderConductor", "External", "DirectXShaderCompiler", "include"));
