Silenced warning about failing to load Slate textures when running a commandlet, as Slate does not initialize/render for commandlets

#rnx
#Synced-CL 52288889

diff --git a/Engine/Source/Developer/StandaloneRenderer/Private/OpenGL/SlateOpenGLTextureManager.cpp b/Engine/Source/Developer/StandaloneRenderer/Private/OpenGL/SlateOpenGLTextureManager.cpp
--- a/Engine/Source/Developer/StandaloneRenderer/Private/OpenGL/SlateOpenGLTextureManager.cpp
+++ b/Engine/Source/Developer/StandaloneRenderer/Private/OpenGL/SlateOpenGLTextureManager.cpp
@@ -164,7 +164,7 @@ bool FSlateOpenGLTextureManager::LoadTexture(const FSlateBrush& InBrush, uint32&
 
 	// Load the resource into an array
 	TArray<uint8> Buffer;
-	if (FFileHelper::LoadFileToArray(Buffer, *ResourcePath))
+	if (FFileHelper::LoadFileToArray(Buffer, *ResourcePath, FILEREAD_Silent))
 	{
 		IImageWrapperModule& ImageWrapperModule = FModuleManager::LoadModuleChecked<IImageWrapperModule>(FName("ImageWrapper"));
 
@@ -201,7 +201,7 @@ bool FSlateOpenGLTextureManager::LoadTexture(const FSlateBrush& InBrush, uint32&
 	}
 	else
 	{
-		UE_LOG(LogSlateOpenGL, Log, TEXT("Could not find file for Slate texture: [%s] '%s'"), *InBrush.GetResourceName().ToString(), *ResourcePath);
+		UE_CLOG(!IsRunningCommandlet(), LogSlateOpenGL, Log, TEXT("Could not load file for Slate resource: [%s] '%s'"), *InBrush.GetResourceName().ToString(), *ResourcePath);
 	}
 
 
diff --git a/Engine/Source/Developer/StandaloneRenderer/Private/Windows/D3D/SlateD3DTextureManager.cpp b/Engine/Source/Developer/StandaloneRenderer/Private/Windows/D3D/SlateD3DTextureManager.cpp
--- a/Engine/Source/Developer/StandaloneRenderer/Private/Windows/D3D/SlateD3DTextureManager.cpp
+++ b/Engine/Source/Developer/StandaloneRenderer/Private/Windows/D3D/SlateD3DTextureManager.cpp
@@ -231,7 +231,7 @@ bool FSlateD3DTextureManager::LoadTexture( const FSlateBrush& InBrush, uint32& O
 	uint32 BytesPerPixel = 4;
 	bool bSucceeded = false;
 	TArray<uint8> RawFileData;
-	if( FFileHelper::LoadFileToArray( RawFileData, *ResourcePath ) )
+	if( FFileHelper::LoadFileToArray( RawFileData, *ResourcePath, FILEREAD_Silent) )
 	{
 		IImageWrapperModule& ImageWrapperModule = FModuleManager::LoadModuleChecked<IImageWrapperModule>( FName("ImageWrapper") );
 
@@ -264,7 +264,7 @@ bool FSlateD3DTextureManager::LoadTexture( const FSlateBrush& InBrush, uint32& O
 	}
 	else
 	{
-		UE_LOG(LogSlateD3D, Warning, TEXT("Could not find file for Slate resource: [%s] '%s'"), *InBrush.GetResourceName().ToString(), *ResourcePath);
+		UE_CLOG(!IsRunningCommandlet(), LogSlateD3D, Warning, TEXT("Could not load file for Slate resource: [%s] '%s'"), *InBrush.GetResourceName().ToString(), *ResourcePath);
 	}
 
 	return bSucceeded;
diff --git a/Engine/Source/Runtime/Engine/Private/ImageUtils.cpp b/Engine/Source/Runtime/Engine/Private/ImageUtils.cpp
--- a/Engine/Source/Runtime/Engine/Private/ImageUtils.cpp
+++ b/Engine/Source/Runtime/Engine/Private/ImageUtils.cpp
@@ -98,7 +98,7 @@ bool FImageUtils::SaveImageAutoFormat(const TCHAR * FileName, const FImageView &
 bool FImageUtils::LoadImage(const TCHAR * Filename, FImage & OutImage)
 {
 	TArray64<uint8> Buffer;
-	if ( ! FFileHelper::LoadFileToArray(Buffer, Filename) )
+	if ( ! FFileHelper::LoadFileToArray(Buffer, Filename, FILEREAD_Silent) )
 	{
 		return false;
 	}
@@ -1357,7 +1357,7 @@ UTexture2D* FImageUtils::ImportFileAsTexture2D(const FString& Filename)
 {
 	UTexture2D* NewTexture = nullptr;
 	TArray64<uint8> Buffer;
-	if (FFileHelper::LoadFileToArray(Buffer, *Filename))
+	if (FFileHelper::LoadFileToArray(Buffer, *Filename, FILEREAD_Silent))
 	{
 		// note this make a Transient / PlatformData only Texture (no TextureSource)
 		NewTexture = FImageUtils::ImportBufferAsTexture2D(Buffer);
diff --git a/Engine/Source/Runtime/SlateRHIRenderer/Private/SlateRHIResourceManager.cpp b/Engine/Source/Runtime/SlateRHIRenderer/Private/SlateRHIResourceManager.cpp
--- a/Engine/Source/Runtime/SlateRHIRenderer/Private/SlateRHIResourceManager.cpp
+++ b/Engine/Source/Runtime/SlateRHIRenderer/Private/SlateRHIResourceManager.cpp
@@ -519,7 +519,7 @@ bool FSlateRHIResourceManager::LoadTexture( const FName& TextureName, const FStr
 	FImage Image;
 	if ( ! FImageUtils::LoadImage(*ResourcePath,Image) )
 	{
-		UE_LOG(LogSlate, Log, TEXT("Could not load file for Slate resource: [%s] file: [%s]"), *TextureName.ToString(),*ResourcePath );
+		UE_CLOG(!IsRunningCommandlet(), LogSlate, Log, TEXT("Could not load file for Slate resource: [%s] file: [%s]"), *TextureName.ToString(),*ResourcePath );
 		return false;
 	}
 
